Branch snapshot for CL 3327560

diff --git a/Engine/Source/Runtime/MeshEditingRuntime/EditableMesh.cpp b/Engine/Source/Runtime/MeshEditingRuntime/EditableMesh.cpp
new file mode 100644
--- /dev/null
+++ b/Engine/Source/Runtime/MeshEditingRuntime/EditableMesh.cpp
@@ -0,0 +1,6799 @@
+// Copyright 1998-2015 Epic Games, Inc. All Rights Reserved.
+
+#include "EditableMesh.h"
+#include "EditableMeshChanges.h"
+#include "EditableMeshCustomVersion.h"
+#include "HAL/IConsoleManager.h"
+#include "GeomTools.h"
+#include "mikktspace.h"	// For tangent computations
+
+
+// =========================================================
+// OpenSubdiv support
+// =========================================================
+
+	// @todo mesheditor urgent subdiv: OpenSubdiv currently isn't compiled for 32-bit; we need to do that to support runtime games properly
+
+	// Disable warning C4191: 'type cast' : unsafe conversion
+	#pragma warning(push)
+	#pragma warning(disable:4191)
+
+	#define M_PI PI
+	#define or ||
+	#define and &&
+	#define not !
+
+	#include "far/topologyRefinerFactory.h"
+	#include "far/topologyDescriptor.h"
+	#include "far/topologyRefiner.h"
+	#include "far/primvarRefiner.h"
+	
+	#if defined(_MSC_VER) && USING_CODE_ANALYSIS
+	#pragma warning(push)
+	#pragma warning(disable : 6011)  // Dereferencing NULL pointer 'X'
+	#pragma warning(disable : 6308)  // 'realloc' might return null pointer: assigning null pointer to 'X', which is passed as an argument to 'realloc', will cause the original memory block to be leaked.
+	#pragma warning(disable : 6385)  // Reading invalid data from 'X':  the readable size is 'Y' bytes, but 'Z' bytes may be read.
+	#pragma warning(disable : 6386)  // Buffer overrun while writing to 'X':  the writable size is 'Y' bytes, but 'Z' bytes might be written.
+	#pragma warning(disable : 28182) // Dereferencing NULL pointer. 'X' contains the same NULL value as 'Y' did.
+	#endif
+
+	// for non manifold
+	// #include "hbr/mesh.h"
+	// #include "hbr/catmark.h"
+
+	#if defined(_MSC_VER) && USING_CODE_ANALYSIS
+	#pragma warning(pop)
+	#endif
+
+	#undef and
+	#undef not
+	#undef M_PI
+
+	// Restore warning C4191.
+	#pragma warning(pop)
+
+
+namespace EditableMesh
+{
+	static FAutoConsoleVariable InterpolatePositionsToLimit( TEXT( "EditableMesh.InterpolatePositionsToLimit" ), 1, TEXT( "Whether to interpolate vertex positions for subdivision meshes all the way to their limit surface position.  Otherwise, we stop at the most refined mesh position." ) );
+	static FAutoConsoleVariable InterpolateFVarsToLimit( TEXT( "EditableMesh.InterpolateFVarsToLimit" ), 1, TEXT( "Whether to interpolate face-varying vertex data for subdivision meshes all the way to their limit surface position.  Otherwise, we stop at the most refined mesh." ) );
+}
+
+
+//
+// =========================================================
+//
+
+const FElementID FElementID::Invalid( TNumericLimits<uint32>::Max() );
+const FVertexID FVertexID::Invalid( TNumericLimits<uint32>::Max() );
+const FEdgeID FEdgeID::Invalid( TNumericLimits<uint32>::Max() );
+const FSectionID FSectionID::Invalid( TNumericLimits<uint32>::Max() );
+const FPolygonID FPolygonID::Invalid( TNumericLimits<uint32>::Max() );
+const FPolygonRef FPolygonRef::Invalid( FSectionID::Invalid, FPolygonID::Invalid );
+
+const FName UEditableMeshAttribute::VertexPositionName( "VertexPosition" );
+const FName UEditableMeshAttribute::VertexCornerSharpnessName( "VertexCornerSharpness" );
+const FName UEditableMeshAttribute::VertexNormalName( "VertexNormal" );
+const FName UEditableMeshAttribute::VertexTangentName( "VertexTangent" );
+const FName UEditableMeshAttribute::VertexBinormalSignName( "VertexBinormalSign" );
+const FName UEditableMeshAttribute::VertexTextureCoordinateName( "VertexTextureCoordinate" );
+const FName UEditableMeshAttribute::VertexColorName( "VertexColor" );
+const FName UEditableMeshAttribute::EdgeIsHardName( "EdgeIsHard" );
+const FName UEditableMeshAttribute::EdgeCreaseSharpnessName( "EdgeCreaseSharpness" );
+
+
+UEditableMesh::UEditableMesh()
+{
+}
+
+
+void UEditableMesh::Serialize( FArchive& Ar )
+{
+	Super::Serialize( Ar );
+	Ar.UsingCustomVersion( FEditableMeshCustomVersion::GUID );
+}
+
+
+void UEditableMesh::PostLoad()
+{
+	Super::PostLoad();
+
+	if( IsPreviewingSubdivisions() )
+	{
+		RefreshOpenSubdiv();
+		RebuildRenderMesh();
+	}
+}
+
+
+const FEditableMeshSubMeshAddress& UEditableMesh::GetSubMeshAddress() const
+{
+	return SubMeshAddress;
+}
+
+
+void UEditableMesh::SetSubMeshAddress( const FEditableMeshSubMeshAddress& NewSubMeshAddress )
+{
+	SubMeshAddress = NewSubMeshAddress;
+}
+
+
+const TArray<FName>& UEditableMesh::GetValidVertexAttributes()
+{
+	static TArray<FName> ValidVertexAttributes;
+	if( ValidVertexAttributes.Num() == 0 )
+	{
+		ValidVertexAttributes.Add( UEditableMeshAttribute::VertexPosition() );
+		ValidVertexAttributes.Add( UEditableMeshAttribute::VertexCornerSharpness() );
+	}
+	return ValidVertexAttributes;
+}
+
+
+const TArray<FName>& UEditableMesh::GetValidPolygonVertexAttributes()
+{
+	static TArray<FName> ValidPolygonVertexAttributes;
+	if( ValidPolygonVertexAttributes.Num() == 0 )
+	{
+		ValidPolygonVertexAttributes.Add( UEditableMeshAttribute::VertexNormal() );
+		ValidPolygonVertexAttributes.Add( UEditableMeshAttribute::VertexTangent() );
+		ValidPolygonVertexAttributes.Add( UEditableMeshAttribute::VertexBinormalSign() );
+		ValidPolygonVertexAttributes.Add( UEditableMeshAttribute::VertexTextureCoordinate() );
+		ValidPolygonVertexAttributes.Add( UEditableMeshAttribute::VertexColor() );
+	}
+	return ValidPolygonVertexAttributes;
+}
+
+
+const TArray<FName>& UEditableMesh::GetValidEdgeAttributes()
+{
+	static TArray<FName> ValidEdgeAttributes;
+	if( ValidEdgeAttributes.Num() == 0 )
+	{
+		ValidEdgeAttributes.Add( UEditableMeshAttribute::EdgeIsHard() );
+		ValidEdgeAttributes.Add( UEditableMeshAttribute::EdgeCreaseSharpness() );
+	}
+	return ValidEdgeAttributes;
+}
+
+
+int32 UEditableMesh::GetMaxAttributeIndex( const FName AttributeName ) const
+{
+	if( AttributeName == UEditableMeshAttribute::VertexTextureCoordinate() )
+	{
+		return GetTextureCoordinateCount();
+	}
+
+	return 1;
+}
+
+
+FSectionID UEditableMesh::GetFirstValidSection() const
+{
+	FSectionID FirstValidSectionID = FSectionID::Invalid;
+
+	// @todo mesheditor perf: We can do this without iterating if we make the function virtual and allow the backend to return the first valid section
+	const int32 SectionArraySize = GetSectionArraySize();
+	for( int32 SectionIndex = 0; SectionIndex < SectionArraySize; ++SectionIndex )
+	{
+		const FSectionID SectionID( SectionIndex );
+		if( IsValidSection( SectionID ) )
+		{
+			FirstValidSectionID = SectionID;
+		}
+	}
+
+	return FirstValidSectionID;
+}
+
+
+int32 UEditableMesh::GetTotalPolygonCount() const
+{
+	int32 TotalPolygonCount = 0;
+
+	const int32 SectionArraySize = GetSectionArraySize();
+	for( int32 SectionIndex = 0; SectionIndex < SectionArraySize; ++SectionIndex )
+	{
+		const FSectionID SectionID( SectionIndex );
+		if( IsValidSection( SectionID ) )
+		{
+			const int32 PolygonCount = GetPolygonCount( SectionID );
+			TotalPolygonCount += PolygonCount;
+		}
+	}
+
+	return TotalPolygonCount;
+}
+
+
+int32 UEditableMesh::GetTextureCoordinateCount() const
+{
+	return this->TextureCoordinateCount;
+}
+
+
+int32 UEditableMesh::GetSubdivisionCount() const
+{
+	return this->SubdivisionCount;
+}
+
+
+bool UEditableMesh::IsPreviewingSubdivisions() const
+{
+	return GetSubdivisionCount() > 0;		// @todo mesheditor subdiv: Make optional even when subdivisions are active
+}
+
+
+void UEditableMesh::GetVertexConnectedPolygons( const FVertexID VertexID, TArray<FPolygonRef>& OutConnectedPolygonRefs ) const
+{
+	OutConnectedPolygonRefs.Reset();
+
+	const int32 ConnectedEdgeCount = GetVertexConnectedEdgeCount( VertexID );
+	for( int32 EdgeNumber = 0; EdgeNumber < ConnectedEdgeCount; ++EdgeNumber )
+	{
+		const FEdgeID ConnectedEdgeID = GetVertexConnectedEdge( VertexID, EdgeNumber );
+
+		const int32 ConnectedPolygonCount = GetEdgeConnectedPolygonCount( ConnectedEdgeID );
+		for( int32 PolygonNumber = 0; PolygonNumber < ConnectedPolygonCount; ++PolygonNumber )
+		{
+			const FPolygonRef PolygonRef = GetEdgeConnectedPolygon( ConnectedEdgeID, PolygonNumber );
+
+			// Add uniquely, because a vertex's connected edges may neighbor the same polygon
+			OutConnectedPolygonRefs.AddUnique( PolygonRef );
+		}
+	}
+}
+
+
+void UEditableMesh::GetVertexAdjacentVertices( const FVertexID VertexID, TArray< FVertexID >& OutAdjacentVertexIDs ) const
+{
+	OutAdjacentVertexIDs.Reset();
+
+	const int32 ConnectedEdgeCount = GetVertexConnectedEdgeCount( VertexID );
+	for( int32 EdgeNumber = 0; EdgeNumber < ConnectedEdgeCount; ++EdgeNumber )
+	{
+		const FEdgeID ConnectedEdgeID = GetVertexConnectedEdge( VertexID, EdgeNumber );
+		
+		FVertexID EdgeVertexIDs[ 2 ];
+		GetEdgeVertices( ConnectedEdgeID, /* Out */ EdgeVertexIDs[0], /* Out */ EdgeVertexIDs[1] );
+
+		OutAdjacentVertexIDs.Add( EdgeVertexIDs[ 0 ] == VertexID ? EdgeVertexIDs[ 1 ] : EdgeVertexIDs[ 0 ] );
+	}
+}
+
+
+void UEditableMesh::GetEdgeVertices( const FEdgeID EdgeID, FVertexID& OutEdgeVertexID0, FVertexID& OutEdgeVertexID1 ) const
+{
+	OutEdgeVertexID0 = GetEdgeVertex( EdgeID, 0 );
+	OutEdgeVertexID1 = GetEdgeVertex( EdgeID, 1 );
+}
+
+
+void UEditableMesh::GetEdgeConnectedPolygons( const FEdgeID EdgeID, TArray<FPolygonRef>& OutConnectedPolygonRefs ) const
+{
+	OutConnectedPolygonRefs.Reset();
+
+	const int32 ConnectedPolygonCount = GetEdgeConnectedPolygonCount( EdgeID );
+	for( int32 PolygonNumber = 0; PolygonNumber < ConnectedPolygonCount; ++PolygonNumber )
+	{
+		const FPolygonRef PolygonRef = GetEdgeConnectedPolygon( EdgeID, PolygonNumber );
+
+		// Add uniquely, because a vertex's connected edges may neighbor the same polygon
+		OutConnectedPolygonRefs.AddUnique( PolygonRef );
+	}
+}
+
+
+void UEditableMesh::GetEdgeLoopElements( const FEdgeID EdgeID, TArray<FEdgeID>& EdgeLoopIDs ) const
+{
+	EdgeLoopIDs.Reset();
+
+	// Maintain a list of unique edge IDs which form the loop
+	static TSet<FEdgeID> EdgeIDs;
+	EdgeIDs.Reset();
+	
+	// Maintain a stack of edges to be processed, in lieu of recursion.
+	// We also store which vertex of the edge has already been processed (so we don't retrace our steps when processing stack items).
+	static TArray<TTuple<FEdgeID, FVertexID>> EdgeStack;
+	EdgeStack.Reset();
+	EdgeStack.Push( MakeTuple( EdgeID, FVertexID::Invalid ) );
+
+	// Process edge IDs on the stack
+	while( EdgeStack.Num() > 0 )
+	{
+		const TTuple<FEdgeID, FVertexID> CurrentItem = EdgeStack.Pop();
+		const FEdgeID CurrentEdgeID = CurrentItem.Get<0>();
+		const FVertexID FromVertexID = CurrentItem.Get<1>();
+		EdgeIDs.Add( CurrentEdgeID );
+
+		// Get the polygons connected to this edge. When continuing the loop, the criterion is that new edges must share no polygons with this edge,
+		// i.e. they are the other side of a perpendicular edge.
+		static TArray<FPolygonRef> ConnectedPolygons;
+		GetEdgeConnectedPolygons( CurrentEdgeID, ConnectedPolygons );
+
+		// Now look for edges connected to each end of this edge
+		for( int32 ConnectedVertexIndex = 0; ConnectedVertexIndex < 2; ++ConnectedVertexIndex )
+		{
+			// If we have already processed this vertex, skip it
+			const FVertexID ConnectedVertexID = GetEdgeVertex( CurrentEdgeID, ConnectedVertexIndex );
+			if( ConnectedVertexID == FromVertexID )
+			{
+				continue;
+			}
+
+			// This is the candidate edge ID which continues the loop beyond the vertex being processed
+			FEdgeID AdjacentEdgeID = FEdgeID::Invalid;
+
+			// Iterate through all edges connected to this vertex
+			const int32 VertexConnectedEdgeCount = GetVertexConnectedEdgeCount( ConnectedVertexID );
+			for( int32 EdgeIndex = 0; EdgeIndex < VertexConnectedEdgeCount; ++EdgeIndex )
+			{
+				const FEdgeID ConnectedEdgeID = GetVertexConnectedEdge( ConnectedVertexID, EdgeIndex );
+
+				// If this edge hasn't been added to the loop...
+				if( !EdgeIDs.Contains( ConnectedEdgeID ) )
+				{
+					// ...see if it shares any polygons with the original edge (intersection operation)
+					const int32 EdgeConnectedPolygonsCount = GetEdgeConnectedPolygonCount( ConnectedEdgeID );
+					bool bIsCandidateEdge = true;
+					for( int32 PolygonIndex = 0; PolygonIndex < EdgeConnectedPolygonsCount; ++PolygonIndex )
+					{
+						const FPolygonRef ConnectedPolygonID = GetEdgeConnectedPolygon( ConnectedEdgeID, PolygonIndex );
+						if( ConnectedPolygons.Contains( ConnectedPolygonID ) )
+						{
+							bIsCandidateEdge = false;
+							break;
+						}
+					}
+
+					// We have found an edge connected to this vertex which doesn't share any polys with the original edge
+					if( bIsCandidateEdge )
+					{
+						if( AdjacentEdgeID == FEdgeID::Invalid )
+						{
+							// If it's the first such edge which meets the criteria, remember it
+							AdjacentEdgeID = ConnectedEdgeID;
+						}
+						else
+						{
+							// If we already have a possible edge, stop the loop here; we don't allow splits in the loop if there is more than one candidate
+							AdjacentEdgeID = FEdgeID::Invalid;
+							break;
+						}
+					}
+				}
+			}
+
+			if( AdjacentEdgeID != FEdgeID::Invalid )
+			{
+				EdgeStack.Push( MakeTuple( AdjacentEdgeID, ConnectedVertexID ) );
+			}
+		}
+	}
+
+	for( const FEdgeID EdgeLoopID : EdgeIDs )
+	{
+		EdgeLoopIDs.Add( EdgeLoopID );
+	}
+}
+
+
+int32 UEditableMesh::GetPolygonPerimeterEdgeCount( const FPolygonRef PolygonRef ) const
+{
+	// All polygons have the same number of edges as they do vertices
+	return GetPolygonPerimeterVertexCount( PolygonRef );
+}
+
+
+int32 UEditableMesh::GetPolygonHoleEdgeCount( const FPolygonRef PolygonRef, const int32 HoleNumber ) const
+{
+	// All polygons have the same number of edges as they do vertices
+	return GetPolygonHoleVertexCount( PolygonRef, HoleNumber );
+}
+
+
+void UEditableMesh::GetPolygonPerimeterVertices( const FPolygonRef PolygonRef, TArray<FVertexID>& OutPolygonPerimeterVertexIDs ) const
+{
+	const int32 NumPolygonPerimeterVertices = GetPolygonPerimeterVertexCount( PolygonRef );
+
+	OutPolygonPerimeterVertexIDs.SetNumUninitialized( NumPolygonPerimeterVertices, false );
+	for( int32 VertexNumber = 0; VertexNumber < NumPolygonPerimeterVertices; ++VertexNumber )
+	{
+		OutPolygonPerimeterVertexIDs[ VertexNumber ] = GetPolygonPerimeterVertex( PolygonRef, VertexNumber );
+	}
+}
+
+
+void UEditableMesh::GetPolygonHoleVertices( const FPolygonRef PolygonRef, const int32 HoleNumber, TArray<FVertexID>& OutHoleVertexIDs ) const
+{
+	const int32 NumHoleVertices = GetPolygonHoleVertexCount( PolygonRef, HoleNumber );
+
+	OutHoleVertexIDs.SetNumUninitialized( NumHoleVertices, false );
+	for( int32 VertexNumber = 0; VertexNumber < NumHoleVertices; ++VertexNumber )
+	{
+		OutHoleVertexIDs[ VertexNumber ] = GetPolygonHoleVertex( PolygonRef, HoleNumber, VertexNumber );
+	}
+}
+
+
+FEdgeID UEditableMesh::GetPolygonPerimeterEdge( const FPolygonRef PolygonRef, const int32 PerimeterEdgeNumber, bool& bOutEdgeWindingIsReversedForPolygon ) const
+{
+	FEdgeID FoundEdgeID = FEdgeID::Invalid;
+	bool bFoundEdge = false;
+	bOutEdgeWindingIsReversedForPolygon = false;
+
+	const int32 NumPolygonPerimeterEdges = GetPolygonPerimeterEdgeCount( PolygonRef );
+	const int32 NumPolygonPerimeterVertices = NumPolygonPerimeterEdges;		// Edge and vertex count are always the same
+
+	check( NumPolygonPerimeterEdges > 0 );
+	for( int32 CurrentPerimeterEdgeNumber = 0; CurrentPerimeterEdgeNumber < NumPolygonPerimeterEdges; ++CurrentPerimeterEdgeNumber )
+	{
+		const int32 PerimeterVertexNumber = CurrentPerimeterEdgeNumber;	// Edge and vertex counts are always the same
+
+		const FVertexID VertexID = GetPolygonPerimeterVertex( PolygonRef, PerimeterVertexNumber );
+		const FVertexID NextVertexID = GetPolygonPerimeterVertex( PolygonRef, ( PerimeterVertexNumber + 1 ) % NumPolygonPerimeterVertices );
+
+		// Find the edge that connects these vertices
+		bool bFoundCurrentEdge = false;
+
+		const int32 NumVertexConnectedEdges = GetVertexConnectedEdgeCount( VertexID );	// @todo mesheditor perf: Could be made faster by always iterating over the vertex with the fewest number of connected edges, instead of the first edge vertex
+		for( int32 VertexEdgeNumber = 0; VertexEdgeNumber < NumVertexConnectedEdges; ++VertexEdgeNumber )
+		{
+			const FEdgeID VertexConnectedEdgeID = GetVertexConnectedEdge( VertexID, VertexEdgeNumber );
+
+			// Try the edge's first vertex.  Does it point to us?
+			FVertexID OtherEdgeVertexID = GetEdgeVertex( VertexConnectedEdgeID, 0 );
+			bOutEdgeWindingIsReversedForPolygon = false;
+			if( OtherEdgeVertexID == VertexID )
+			{
+				// Must be the the other one
+				OtherEdgeVertexID = GetEdgeVertex( VertexConnectedEdgeID, 1 );
+				bOutEdgeWindingIsReversedForPolygon = true;
+			}
+
+			if( OtherEdgeVertexID == NextVertexID )
+			{
+				// We found the edge!
+				FoundEdgeID = VertexConnectedEdgeID;
+				bFoundCurrentEdge = true;
+				break;
+			}
+		}
+		check( bFoundCurrentEdge );
+
+		if( CurrentPerimeterEdgeNumber == PerimeterEdgeNumber )
+		{
+			bFoundEdge = true;
+			break;
+		}
+	}
+
+	check( bFoundEdge );
+	return FoundEdgeID;
+}
+
+
+FEdgeID UEditableMesh::GetPolygonHoleEdge( const FPolygonRef PolygonRef, const int32 HoleNumber, const int32 HoleEdgeNumber ) const
+{
+	bool bFoundEdge = false;
+	FEdgeID FoundEdgeID = FEdgeID::Invalid;
+
+	const int32 NumHoleEdges = GetPolygonHoleEdgeCount( PolygonRef, HoleNumber );
+	const int32 NumHoleVertices = NumHoleEdges;		// Edge and vertex count are always the same
+
+	check( NumHoleEdges > 0 );
+	for( int32 CurrentHoleEdgeNumber = 0; CurrentHoleEdgeNumber < NumHoleEdges; ++CurrentHoleEdgeNumber )
+	{
+		const int32 HoleVertexNumber = CurrentHoleEdgeNumber;	// Edge and vertex counts are always the same
+
+		const FVertexID VertexID = GetPolygonHoleVertex( PolygonRef, HoleNumber, HoleVertexNumber );
+		const FVertexID NextVertexID = GetPolygonHoleVertex( PolygonRef, HoleNumber, ( HoleVertexNumber + 1 ) % NumHoleVertices );
+
+		// Find the edge that connects these vertices
+		bool bFoundCurrentEdge = false;
+
+		const int32 NumVertexConnectedEdges = GetVertexConnectedEdgeCount( VertexID );	// @todo mesheditor perf: Could be made faster by always iterating over the vertex with the fewest number of connected edges, instead of the first edge vertex
+		for( int32 VertexEdgeNumber = 0; VertexEdgeNumber < NumVertexConnectedEdges; ++VertexEdgeNumber )
+		{
+			const FEdgeID VertexConnectedEdgeID = GetVertexConnectedEdge( VertexID, VertexEdgeNumber );
+
+			// Try the edge's first vertex.  Does it point to us?
+			FVertexID OtherEdgeVertexID = GetEdgeVertex( VertexConnectedEdgeID, 0 );
+			if( OtherEdgeVertexID == VertexID )
+			{
+				// Must be the the other one
+				OtherEdgeVertexID = GetEdgeVertex( VertexConnectedEdgeID, 1 );
+			}
+
+			if( OtherEdgeVertexID == NextVertexID )
+			{
+				// We found the edge!
+				FoundEdgeID = VertexConnectedEdgeID;
+				bFoundCurrentEdge = true;
+				break;
+			}
+		}
+		check( bFoundCurrentEdge );
+
+		if( CurrentHoleEdgeNumber == HoleEdgeNumber )
+		{
+			bFoundEdge = true;
+			break;
+		}
+	}
+
+	check( bFoundEdge );
+	return FoundEdgeID;
+}
+
+
+void UEditableMesh::GetPolygonPerimeterEdges( const FPolygonRef PolygonRef, TArray<FEdgeID>& OutPolygonPerimeterEdgeIDs ) const
+{
+	const int32 NumPolygonPerimeterEdges = GetPolygonPerimeterEdgeCount( PolygonRef );
+	const int32 NumPolygonPerimeterVertices = NumPolygonPerimeterEdges;		// Edge and vertex count are always the same
+
+	OutPolygonPerimeterEdgeIDs.SetNumUninitialized( NumPolygonPerimeterEdges, false );
+
+	for( int32 PerimeterEdgeNumber = 0; PerimeterEdgeNumber < NumPolygonPerimeterEdges; ++PerimeterEdgeNumber )
+	{
+		const int32 PerimeterVertexNumber = PerimeterEdgeNumber;	// Edge and vertex counts are always the same
+
+		const FVertexID VertexID = GetPolygonPerimeterVertex( PolygonRef, PerimeterVertexNumber );
+		const FVertexID NextVertexID = GetPolygonPerimeterVertex( PolygonRef, ( PerimeterVertexNumber + 1 ) % NumPolygonPerimeterVertices );
+
+		// Find the edge that connects these vertices
+		FEdgeID FoundEdgeID = FEdgeID::Invalid;
+		bool bFoundEdge = false;
+
+		const int32 NumVertexConnectedEdges = GetVertexConnectedEdgeCount( VertexID );	// @todo mesheditor perf: Could be made faster by always iterating over the vertex with the fewest number of connected edges, instead of the first edge vertex
+		for( int32 VertexEdgeNumber = 0; VertexEdgeNumber < NumVertexConnectedEdges; ++VertexEdgeNumber )
+		{
+			const FEdgeID VertexConnectedEdgeID = GetVertexConnectedEdge( VertexID, VertexEdgeNumber );
+
+			// Try the edge's first vertex.  Does it point to our next edge?
+			FVertexID OtherEdgeVertexID = GetEdgeVertex( VertexConnectedEdgeID, 0 );
+			if( OtherEdgeVertexID == VertexID )
+			{
+				// Must be the the other one
+				OtherEdgeVertexID = GetEdgeVertex( VertexConnectedEdgeID, 1 );
+			}
+			else
+			{
+				check( GetEdgeVertex( VertexConnectedEdgeID, 1 ) == VertexID );
+			}
+
+			if( OtherEdgeVertexID == NextVertexID )
+			{
+				// We found the edge!
+				FoundEdgeID = VertexConnectedEdgeID;
+				bFoundEdge = true;
+				break;
+			}
+		}
+
+		// @todo mesheditor perf: A lot of these inner-loop check() statements need to be disabled (or changes to checkSlow()) for faster performance
+		check( bFoundEdge );
+
+		OutPolygonPerimeterEdgeIDs[ PerimeterEdgeNumber ] = FoundEdgeID;
+	}
+}
+
+
+void UEditableMesh::GetPolygonHoleEdges( const FPolygonRef PolygonRef, const int32 HoleNumber, TArray<FEdgeID>& OutHoleEdgeIDs ) const
+{
+	const int32 NumHoleEdges = GetPolygonHoleEdgeCount( PolygonRef, HoleNumber );
+	const int32 NumHoleVertices = NumHoleEdges;		// Edge and vertex count are always the same
+
+	OutHoleEdgeIDs.SetNumUninitialized( NumHoleEdges, false );
+
+	for( int32 HoleEdgeNumber = 0; HoleEdgeNumber < NumHoleEdges; ++HoleEdgeNumber )
+	{
+		const int32 HoleVertexNumber = HoleEdgeNumber;	// Edge and vertex counts are always the same
+
+		const FVertexID VertexID = GetPolygonHoleVertex( PolygonRef, HoleNumber, HoleVertexNumber );
+		const FVertexID NextVertexID = GetPolygonHoleVertex( PolygonRef, HoleNumber, ( HoleVertexNumber + 1 ) % NumHoleVertices );
+
+		// Find the edge that connects these vertices
+		FEdgeID FoundEdgeID = FEdgeID::Invalid;
+		bool bFoundEdge = false;
+
+		const int32 NumVertexConnectedEdges = GetVertexConnectedEdgeCount( VertexID );	// @todo mesheditor perf: Could be made faster by always iterating over the vertex with the fewest number of connected edges, instead of the first edge vertex
+		for( int32 VertexEdgeNumber = 0; VertexEdgeNumber < NumVertexConnectedEdges; ++VertexEdgeNumber )
+		{
+			const FEdgeID VertexConnectedEdgeID = GetVertexConnectedEdge( VertexID, VertexEdgeNumber );
+
+			// Try the edge's first vertex.  Does it point to us?
+			FVertexID OtherEdgeVertexID = GetEdgeVertex( VertexConnectedEdgeID, 0 );
+			if( OtherEdgeVertexID == VertexID )
+			{
+				// Must be the the other one
+				OtherEdgeVertexID = GetEdgeVertex( VertexConnectedEdgeID, 1 );
+			}
+
+			if( OtherEdgeVertexID == NextVertexID )
+			{
+				// We found the edge!
+				FoundEdgeID = VertexConnectedEdgeID;
+				bFoundEdge = true;
+				break;
+			}
+		}
+		check( bFoundEdge );
+
+		OutHoleEdgeIDs[ HoleEdgeNumber ] = FoundEdgeID;
+	}
+}
+
+
+void UEditableMesh::GetPolygonAdjacentPolygons( const FPolygonRef PolygonRef, TArray<FPolygonRef>& OutAdjacentPolygons ) const
+{
+	OutAdjacentPolygons.Reset();
+
+	static TArray<FEdgeID> PolygonPerimeterEdges;
+	GetPolygonPerimeterEdges( PolygonRef, PolygonPerimeterEdges );
+
+	for( const FEdgeID EdgeID : PolygonPerimeterEdges )
+	{
+		static TArray<FPolygonRef> EdgeConnectedPolygons;
+		GetEdgeConnectedPolygons( EdgeID, EdgeConnectedPolygons );
+
+		for( const FPolygonRef EdgeConnectedPolygon : EdgeConnectedPolygons )
+		{
+			if( EdgeConnectedPolygon != PolygonRef )
+			{
+				OutAdjacentPolygons.AddUnique( EdgeConnectedPolygon );
+			}
+		}
+	}
+}
+
+
+FBox UEditableMesh::ComputeBoundingBox() const
+{
+	FBox BoundingBox;
+	BoundingBox.Init();
+
+	const int32 VertexArraySize = GetVertexArraySize();
+	for( int32 VertexNumber = 0; VertexNumber < VertexArraySize; ++VertexNumber )
+	{
+		const FVertexID VertexID( VertexNumber );
+		if( IsValidVertex( VertexID ) )
+		{
+			const FVector VertexPosition = GetVertexAttribute( VertexID, UEditableMeshAttribute::VertexPosition(), 0 );
+			BoundingBox += VertexPosition;
+		}
+	}
+
+	return BoundingBox;
+}
+
+
+FBoxSphereBounds UEditableMesh::ComputeBoundingBoxAndSphere() const
+{
+	const FBox BoundingBox = ComputeBoundingBox();
+
+	FBoxSphereBounds BoundingBoxAndSphere;
+	BoundingBox.GetCenterAndExtents( /* Out */ BoundingBoxAndSphere.Origin, /* Out */ BoundingBoxAndSphere.BoxExtent );
+
+	// Calculate the bounding sphere, using the center of the bounding box as the origin.
+	BoundingBoxAndSphere.SphereRadius = 0.0f;
+
+	const int32 VertexArraySize = GetVertexArraySize();
+	for( int32 VertexNumber = 0; VertexNumber < VertexArraySize; ++VertexNumber )
+	{
+		const FVertexID VertexID( VertexNumber );
+		if( IsValidVertex( VertexID ) )
+		{
+			const FVector VertexPosition = GetVertexAttribute( VertexID, UEditableMeshAttribute::VertexPosition(), 0 );
+
+			BoundingBoxAndSphere.SphereRadius = FMath::Max( ( VertexPosition - BoundingBoxAndSphere.Origin ).Size(), BoundingBoxAndSphere.SphereRadius );
+		}
+	}
+
+	return BoundingBoxAndSphere;
+}
+
+
+FVector UEditableMesh::ComputePolygonCenter( const FPolygonRef PolygonRef ) const
+{
+	FVector Centroid = FVector::ZeroVector;
+
+	static TArray<FVertexID> PerimeterVertexIDs;
+	GetPolygonPerimeterVertices( PolygonRef, /* Out */ PerimeterVertexIDs );
+
+	for( const FVertexID VertexID : PerimeterVertexIDs )
+	{
+		const FVector Position = GetVertexAttribute( VertexID, UEditableMeshAttribute::VertexPosition(), 0 );
+
+		Centroid += Position;
+	}
+
+	return Centroid / (float)PerimeterVertexIDs.Num();
+}
+
+
+FPlane UEditableMesh::ComputePolygonPlane( const FPolygonRef PolygonRef ) const
+{
+	// NOTE: This polygon plane computation code is partially based on the implementation of "Newell's method" from Real-Time 
+	//       Collision Detection by Christer Ericson, published by Morgan Kaufmann Publishers, (c) 2005 Elsevier Inc
+
+	// @todo mesheditor perf: For polygons that are just triangles, use a cross product to get the normal fast!
+	// @todo mesheditor perf: We could skip computing the plane distance when we only need the normal
+	// @todo mesheditor perf: We could cache these computed polygon normals; or just use the normal of the first three vertices' triangle if it is satisfactory in all cases
+
+	FVector Centroid = FVector::ZeroVector;
+	FVector Normal = FVector::ZeroVector;
+
+	static TArray<FVertexID> PerimeterVertexIDs;
+	GetPolygonPerimeterVertices( PolygonRef, /* Out */ PerimeterVertexIDs );
+
+	// Use 'Newell's Method' to compute a robust 'best fit' plane from the vertices of this polygon
+	for( int32 VertexNumberI = PerimeterVertexIDs.Num() - 1, VertexNumberJ = 0; VertexNumberJ < PerimeterVertexIDs.Num(); VertexNumberI = VertexNumberJ, VertexNumberJ++ )
+	{
+		const FVertexID VertexIDI = PerimeterVertexIDs[ VertexNumberI ];
+		const FVector PositionI = GetVertexAttribute( VertexIDI, UEditableMeshAttribute::VertexPosition(), 0 );
+
+		const FVertexID VertexIDJ = PerimeterVertexIDs[ VertexNumberJ ];
+		const FVector PositionJ = GetVertexAttribute( VertexIDJ, UEditableMeshAttribute::VertexPosition(), 0 );
+
+		Centroid += PositionJ;
+
+		Normal.X += ( PositionJ.Y - PositionI.Y ) * ( PositionI.Z + PositionJ.Z );
+		Normal.Y += ( PositionJ.Z - PositionI.Z ) * ( PositionI.X + PositionJ.X );
+		Normal.Z += ( PositionJ.X - PositionI.X ) * ( PositionI.Y + PositionJ.Y );
+	}
+
+	Normal.Normalize();
+
+	// Construct a plane from the normal and centroid
+	return FPlane( Normal, FVector::DotProduct( Centroid, Normal ) / (float)PerimeterVertexIDs.Num() );
+}
+
+
+FVector UEditableMesh::ComputePolygonNormal( const FPolygonRef PolygonRef ) const
+{
+	const FPlane PolygonPlane = ComputePolygonPlane( PolygonRef );
+	const FVector PolygonNormal( PolygonPlane.X, PolygonPlane.Y, PolygonPlane.Z );
+	return PolygonNormal;
+}
+
+
+FVector UEditableMesh::ComputePolygonPerimeterVertexNormal( const FPolygonRef PolygonRef, const int32 PolygonVertexNumber ) const
+{
+	// The idea here is that we build a graph of adjacent polygons which include the vertex whose normal we want to compute
+	// (where only polygons separated by a soft edge are considered to be adjacent).
+	// Then we walk the graph starting from our start poly, and sum all the visited poly normals, weighted by the angle they make with the vertex.
+	//
+	// We only consider adjacency local to the vertex being considered - this means we can support topology such as a face only partially delimited
+	// by hard edges. Determining smoothing groups would be less flexible.
+
+	// @todo mesheditor: this is not likely to scale well for large meshes (since all normals are recalculated after topology changes)
+	// as it is of pretty high complexity.
+	// We should either cache further adjacency information (e.g. adjacent polys), or not recalculate normals which have not changed.
+
+	// ID of the vertex whose normal we wish to compute
+	const FVertexID VertexID = GetPolygonPerimeterVertex( PolygonRef, PolygonVertexNumber );
+
+	// For every polygon which includes the vertex, we record which polygons are adjacent, and the angle the polygon makes with the vertex.
+	// When walking the graph later, the bVisited flag will be used to ensure we visit each node only once.
+	struct FPolygonGraphNode
+	{
+		FPolygonGraphNode()
+			: Angle( 0.0f )
+			, bVisited( false )
+		{}
+
+		TArray<FPolygonRef> AdjacentPolygonRefs;
+		float Angle;
+		bool bVisited;
+	};
+
+	// Use this to build a graph of polygons
+	static TMap<FPolygonRef, FPolygonGraphNode> PolygonGraphNodes;
+	PolygonGraphNodes.Empty();
+
+	// Get a list of all polygons which share this vertex.
+	// This is the container we iterate through when building the adjacencies.
+	static TArray<FPolygonRef> ConnectedPolygons;
+	GetVertexConnectedPolygons( VertexID, ConnectedPolygons );
+
+	// Get a list of all soft edges which share this vertex.
+	// We're not interested in hard edges as we don't consider them to denote an adjacency between their neighbour polys.
+	static TArray<FEdgeID> ConnectedSoftEdges;
+	ConnectedSoftEdges.Empty();
+
+	const int32 ConnectedEdgeCount = GetVertexConnectedEdgeCount( VertexID );
+	for( int32 EdgeNumber = 0; EdgeNumber < ConnectedEdgeCount; ++EdgeNumber )
+	{
+		const FEdgeID EdgeID( GetVertexConnectedEdge( VertexID, EdgeNumber ) );
+		const bool bIsSoftEdge = FMath::IsNearlyZero( GetEdgeAttribute( EdgeID, UEditableMeshAttribute::EdgeIsHard(), 0 ).X );
+		if( bIsSoftEdge )
+		{
+			ConnectedSoftEdges.Add( EdgeID );
+		}
+	}
+
+	for( FPolygonRef ConnectedPolygon : ConnectedPolygons )
+	{
+		// For each connected polygon, create a PolygonInfo entry
+		FPolygonGraphNode& PolygonGraphNode = PolygonGraphNodes.Add( ConnectedPolygon, FPolygonGraphNode() );
+
+		// Get all the vertices which form the polygon perimeter
+		// @todo mesheditor: what about holes?
+		static TArray<FVertexID> PolygonVertices;
+		GetPolygonPerimeterVertices( ConnectedPolygon, PolygonVertices );
+		const int32 PolygonVertexCount = PolygonVertices.Num();
+		check( PolygonVertexCount > 2 );
+
+		FVertexID LastVertexID( PolygonVertices[ PolygonVertexCount - 2 ] );
+		FVertexID ThisVertexID( PolygonVertices[ PolygonVertexCount - 1 ] );
+
+		// Go through all the polygon perimeter vertices, looking for the vertex whose normal we are setting
+		for( int32 VertexIndex = 0; VertexIndex < PolygonVertexCount; ++VertexIndex )
+		{
+			const FVertexID NextVertexID( PolygonVertices[ VertexIndex ] );
+
+			if( ThisVertexID == VertexID )
+			{
+				// Found the vertex, and we now have the two adjacent vertices too.
+				const FVector LastVertexPosition( GetVertexAttribute( LastVertexID, UEditableMeshAttribute::VertexPosition(), 0 ) );
+				const FVector ThisVertexPosition( GetVertexAttribute( ThisVertexID, UEditableMeshAttribute::VertexPosition(), 0 ) );
+				const FVector NextVertexPosition( GetVertexAttribute( NextVertexID, UEditableMeshAttribute::VertexPosition(), 0 ) );
+
+				const FVector Direction1 = ( LastVertexPosition - ThisVertexPosition ).GetSafeNormal();
+				const FVector Direction2 = ( NextVertexPosition - ThisVertexPosition ).GetSafeNormal();
+
+				// Record the angle the poly makes with the vertex, to be used as the normal weight
+				PolygonGraphNode.Angle = FMath::Acos( FVector::DotProduct( Direction1, Direction2 ) );
+			}
+			else
+			{
+				// Now examine all the soft edges connected to the vertex, looking for those which are also connected to this perimeter vertex.
+				// We need this to determine adjacent polygons.
+				for( FEdgeID ConnectedEdge : ConnectedSoftEdges )
+				{
+					FVertexID EdgeVertexID0;
+					FVertexID EdgeVertexID1;
+					GetEdgeVertices( ConnectedEdge, EdgeVertexID0, EdgeVertexID1 );
+
+					// If we find an edge which contains this vertex ID, look for adjacent polygons.
+					// We know that the other vertex must be the vertex whose normal we're building.
+					if( EdgeVertexID0 == ThisVertexID || EdgeVertexID1 == ThisVertexID )
+					{
+						check( EdgeVertexID0 == VertexID || EdgeVertexID1 == VertexID );
+						const int32 EdgeConnectedPolygonCount = GetEdgeConnectedPolygonCount( ConnectedEdge );
+						for( int32 EdgeConnectedPolygonIndex = 0; EdgeConnectedPolygonIndex < EdgeConnectedPolygonCount; ++EdgeConnectedPolygonIndex )
+						{
+							// One of the adjacent polygons will be the one we are currently processing.
+							// All the others must be adjacent via the edge we are examining.
+							const FPolygonRef EdgeConnectedPolygonRef = GetEdgeConnectedPolygon( ConnectedEdge, EdgeConnectedPolygonIndex );
+							if( EdgeConnectedPolygonRef != ConnectedPolygon )
+							{
+								PolygonGraphNode.AdjacentPolygonRefs.AddUnique( EdgeConnectedPolygonRef );
+							}
+						}
+					}
+				}
+
+			}
+
+			LastVertexID = ThisVertexID;
+			ThisVertexID = NextVertexID;
+		}
+	}
+
+	// Now walk the graph starting from the polygon being set (a simple non-recursive depth-first search)
+	static TArray<FPolygonRef> GraphNodeStack;
+	GraphNodeStack.Empty( PolygonGraphNodes.Num() );
+	GraphNodeStack.Push( PolygonRef );
+
+	FVector Normal( 0.0f );
+
+	while( GraphNodeStack.Num() > 0 )
+	{
+		const FPolygonRef Node = GraphNodeStack.Pop();
+		FPolygonGraphNode* PolygonGraphNode = PolygonGraphNodes.Find( Node );
+
+		if( PolygonGraphNode && !PolygonGraphNode->bVisited )
+		{
+			// Make node as visited and add normal contribution to result
+			PolygonGraphNode->bVisited = true;
+			Normal += ComputePolygonNormal( Node ) * PolygonGraphNode->Angle;
+
+			// Add adjacent polygons to the node stack to be walked
+			for( FPolygonRef AdjacentPolygonRef : PolygonGraphNode->AdjacentPolygonRefs )
+			{
+				GraphNodeStack.Push( AdjacentPolygonRef );
+			}
+		}
+	}
+
+	return Normal.GetSafeNormal();
+}
+
+
+void UEditableMesh::RefreshOpenSubdiv()
+{
+	OsdTopologyRefiner.Reset();
+
+	if( SubdivisionCount > 0 )
+	{
+		// @todo mesheditor subdiv perf: Ideally we give our topology data straight to OSD rather than have it build it from this simple 
+		// set of parameters.  This will be much more efficient!  In order to do this, we need to specialize TopologyRefinerFactory 
+		// for our mesh (http://graphics.pixar.com/opensubdiv/docs/far_overview.html)
+		OpenSubdiv::Far::TopologyDescriptor OsdTopologyDescriptor;
+		{
+			const int32 VertexArraySize = GetVertexArraySize();
+			OsdTopologyDescriptor.numVertices = VertexArraySize;
+			OsdTopologyDescriptor.numFaces = GetTotalPolygonCount();
+
+			// NOTE: OpenSubdiv likes weights to be between 0.0 and 10.0, so we'll account for that here
+			const float OpenSubdivCreaseWeightMultiplier = 10.0f;
+
+			{
+				// Subdivision corner weights
+				{
+					OsdCornerVertexIndices.Reset();
+					OsdCornerWeights.Reset();
+					for( int32 VertexNumber = 0; VertexNumber < VertexArraySize; ++VertexNumber )
+					{
+						const FVertexID VertexID( VertexNumber );
+						if( IsValidVertex( VertexID ) )
+						{
+							const float VertexCornerSharpness = GetVertexAttribute( VertexID, UEditableMeshAttribute::VertexCornerSharpness(), 0 ).X;
+							if( VertexCornerSharpness > SMALL_NUMBER )
+							{
+								// This vertex is (at least partially) a subdivision corner
+								OsdCornerVertexIndices.Add( VertexNumber );
+								OsdCornerWeights.Add( OpenSubdivCreaseWeightMultiplier * VertexCornerSharpness );
+							}
+						}
+					}
+				}
+
+				// Edge creases
+				{
+					const int32 EdgeArraySize = GetEdgeArraySize();
+
+					OsdCreaseVertexIndexPairs.Reset();
+					OsdCreaseWeights.Reset();
+					for( int32 EdgeNumber = 0; EdgeNumber < EdgeArraySize; ++EdgeNumber )
+					{
+						const FEdgeID EdgeID( EdgeNumber );
+						if( IsValidEdge( EdgeID ) )
+						{
+							const float EdgeCreaseSharpness = GetEdgeAttribute( EdgeID, UEditableMeshAttribute::EdgeCreaseSharpness(), 0 ).X;
+							if( EdgeCreaseSharpness > SMALL_NUMBER )
+							{
+								// This edge is (at least partially) creased
+								FVertexID EdgeVertexID0, EdgeVertexID1;
+								GetEdgeVertices( EdgeID, /* Out */ EdgeVertexID0, /* Out */ EdgeVertexID1 );
+
+								OsdCreaseVertexIndexPairs.Add( EdgeVertexID0.GetValue() );
+								OsdCreaseVertexIndexPairs.Add( EdgeVertexID1.GetValue() );
+								OsdCreaseWeights.Add( OpenSubdivCreaseWeightMultiplier * EdgeCreaseSharpness );
+							}
+						}
+					}
+				}
+
+				OsdNumVerticesPerFace.SetNum( OsdTopologyDescriptor.numFaces, false );
+				OsdVertexIndicesPerFace.Reset();
+				OsdFVarIndicesPerFace.Reset();
+
+				int32 NextOsdFaceIndex = 0;
+
+				const int32 SectionArraySize = GetSectionArraySize();
+				for( int32 SectionNumber = 0; SectionNumber < SectionArraySize; ++SectionNumber )
+				{
+					const FSectionID SectionID( SectionNumber );
+					if( IsValidSection( SectionID ) )
+					{
+						int32 SectionPolygonArraySize = GetPolygonArraySize( SectionID );
+						for( int32 SectionPolygonNumber = 0; SectionPolygonNumber < SectionPolygonArraySize; ++SectionPolygonNumber )
+						{
+							const FPolygonRef PolygonRef( SectionID, FPolygonID( SectionPolygonNumber ) );
+							if( IsValidPolygon( PolygonRef ) )
+							{
+								static TArray<FVertexID> PerimeterVertexIDs;
+								GetPolygonPerimeterVertices( PolygonRef, /* Out */ PerimeterVertexIDs );
+
+								const int32 PerimeterVertexCount = PerimeterVertexIDs.Num();
+								OsdNumVerticesPerFace[ NextOsdFaceIndex++ ] = PerimeterVertexCount;
+
+								for( int32 PerimeterVertexNumber = 0; PerimeterVertexNumber < PerimeterVertexCount; ++PerimeterVertexNumber )
+								{
+									const FVertexID PerimeterVertexID = PerimeterVertexIDs[ PerimeterVertexNumber ];
+
+									OsdVertexIndicesPerFace.Add( PerimeterVertexID.GetValue() );
+									OsdFVarIndicesPerFace.Add( OsdFVarIndicesPerFace.Num() );
+								}
+							}
+						}
+					}
+				}
+
+				check( NextOsdFaceIndex == OsdNumVerticesPerFace.Num()  );
+				check( OsdVertexIndicesPerFace.Num() == OsdFVarIndicesPerFace.Num() );
+			}
+
+			{
+				const int32 TotalFVarChannels = 1;
+				OsdFVarChannels.SetNum( TotalFVarChannels, false );
+				for( int32 FVarChannelNumber = 0; FVarChannelNumber < OsdFVarChannels.Num(); ++FVarChannelNumber )
+				{
+					FOsdFVarChannel& OsdFVarChannel = OsdFVarChannels[ FVarChannelNumber ];
+
+					// @todo mesheditor subdiv: For now, we'll assuming unique face-varying data for every polygon vertex.  Ideally
+					// if we were able to share rendering vertices, we'd want to share face-varying data too.  OpenSubdiv allows
+					// you to pass a separate index array for every channel of face-varying data, but that's probably overkill
+					// for us.  Instead we just need to reflect what we do with rendering vertices.
+					OsdFVarChannel.ValueCount = OsdFVarIndicesPerFace.Num();
+					OsdFVarChannel.ValueIndices = OsdFVarIndicesPerFace.GetData();
+				}
+			}
+
+			OsdTopologyDescriptor.numVertsPerFace = OsdNumVerticesPerFace.GetData();
+			OsdTopologyDescriptor.vertIndicesPerFace = OsdVertexIndicesPerFace.GetData();
+
+			OsdTopologyDescriptor.numCreases = OsdCreaseWeights.Num();
+			OsdTopologyDescriptor.creaseVertexIndexPairs = OsdCreaseVertexIndexPairs.GetData();
+			OsdTopologyDescriptor.creaseWeights = OsdCreaseWeights.GetData();
+
+			OsdTopologyDescriptor.numCorners = OsdCornerWeights.Num();
+			OsdTopologyDescriptor.cornerVertexIndices = OsdCornerVertexIndices.GetData();
+			OsdTopologyDescriptor.cornerWeights = OsdCornerWeights.GetData();
+
+			OsdTopologyDescriptor.numHoles = 0;	// @todo mesheditor holes
+			OsdTopologyDescriptor.holeIndices = nullptr;
+
+			OsdTopologyDescriptor.isLeftHanded = true;
+
+			// Face-varying vertex data.  This maps to our GetPolygonVertexAttribute() calls.
+			OsdTopologyDescriptor.numFVarChannels = OsdFVarChannels.Num();
+			OsdTopologyDescriptor.fvarChannels = reinterpret_cast<OpenSubdiv::v3_0_2::Far::TopologyDescriptor::FVarChannel*>( OsdFVarChannels.GetData() );
+		}
+
+		// We always want Catmull-Clark subdivisions
+		const OpenSubdiv::Sdc::SchemeType OsdSchemeType = OpenSubdiv::Sdc::SCHEME_CATMARK;
+
+		OpenSubdiv::Sdc::Options OsdSdcOptions;
+		{
+			// @todo mesheditor subdiv: Tweak these settings
+			OsdSdcOptions.SetVtxBoundaryInterpolation( OpenSubdiv::Sdc::Options::VTX_BOUNDARY_EDGE_ONLY );
+			OsdSdcOptions.SetFVarLinearInterpolation( OpenSubdiv::Sdc::Options::FVAR_LINEAR_ALL );
+			OsdSdcOptions.SetCreasingMethod( OpenSubdiv::Sdc::Options::CREASE_UNIFORM );
+			OsdSdcOptions.SetTriangleSubdivision( OpenSubdiv::Sdc::Options::TRI_SUB_CATMARK );
+		}
+
+		OpenSubdiv::Far::TopologyRefinerFactory<OpenSubdiv::Far::TopologyDescriptor>::Options OsdTopologyRefinerOptions( OsdSchemeType, OsdSdcOptions );
+
+		this->OsdTopologyRefiner = MakeShareable(
+			OpenSubdiv::Far::TopologyRefinerFactory<OpenSubdiv::Far::TopologyDescriptor>::Create(
+				OsdTopologyDescriptor,
+				OsdTopologyRefinerOptions ) );
+
+		OpenSubdiv::Far::TopologyRefiner::UniformOptions OsdUniformOptions( SubdivisionCount );
+		{
+			// @todo mesheditor subdiv: Should we order child vertices from faces instead of child vertices from vertices first?
+			OsdUniformOptions.orderVerticesFromFacesFirst = false;
+
+			// NOTE: In order for face-varying data to work, OpenSubdiv requires 'fullTopologyInLastLevel' to be enabled.
+			OsdUniformOptions.fullTopologyInLastLevel = true;
+		}
+
+
+		this->OsdTopologyRefiner->RefineUniform( OsdUniformOptions );
+	}
+
+	GenerateOpenSubdivLimitSurfaceData();
+}
+
+
+const FSubdivisionLimitData& UEditableMesh::GetSubdivisionLimitData() const
+{
+	return SubdivisionLimitData;
+}
+
+
+void UEditableMesh::GenerateOpenSubdivLimitSurfaceData()
+{
+	SubdivisionLimitData = FSubdivisionLimitData();
+
+	if( SubdivisionCount > 0 && ensure( this->OsdTopologyRefiner.IsValid() ) )
+	{
+		// Create an OpenSubdiv 'primvar refiner'.  This guy allows us to interpolate data between vertices on a subdivision level.
+		const OpenSubdiv::Far::PrimvarRefiner OsdPrimvarRefiner( *OsdTopologyRefiner );
+
+		struct FOsdVector
+		{
+			void Clear( void* UnusedPtr = nullptr )
+			{
+				Position = FVector( 0, 0, 0 );
+			}
+
+			void AddWithWeight( FOsdVector const& SourceVertexPosition, float Weight )
+			{
+				Position += SourceVertexPosition.Position * Weight;
+			}
+
+			FVector Position;
+		};
+
+		struct FOsdFVarVertexData
+		{
+			void Clear()
+			{
+				TextureCoordinates[0] = TextureCoordinates[1] = FVector2D( 0.0f, 0.0f );
+				VertexColor = FLinearColor( 0.0f, 0.0f, 0.0f, 0.0f );
+			}
+
+			void AddWithWeight( FOsdFVarVertexData const& SourceVertex, float Weight )
+			{
+				TextureCoordinates[ 0 ] += SourceVertex.TextureCoordinates[ 0 ] * Weight;
+				TextureCoordinates[ 1 ] += SourceVertex.TextureCoordinates[ 1 ] * Weight;
+				VertexColor += SourceVertex.VertexColor * Weight;
+			}
+
+			FVector2D TextureCoordinates[2];	// @todo mesheditor subdiv perf: Only two UV sets supported for now (just to avoid heap allocs for dynamic array)
+			FLinearColor VertexColor;
+		};
+
+		const int32 SectionCount = GetSectionCount();
+
+		// Get the limit surface subdivision level from OpenSubdiv
+		const OpenSubdiv::Far::TopologyLevel& OsdLimitLevel = OsdTopologyRefiner->GetLevel( SubdivisionCount );
+
+		const int32 LimitVertexCount = OsdLimitLevel.GetNumVertices();
+		const int32 LimitFaceCount = OsdLimitLevel.GetNumFaces();
+
+		static TArray<FVector> LimitXGradients;
+		LimitXGradients.Reset();
+		static TArray<FVector> LimitYGradients;
+		LimitYGradients.Reset();
+
+		// Grab all of the vertex data and put them in separate contiguous arrays for OpenSubdiv
+		static TArray<FVector> VertexPositions;
+		static TArray<FOsdFVarVertexData> FVarVertexDatas;
+		static TArray<int32> FirstPolygonNumberForSections;
+		{
+			// Vertex positions
+			{
+				// NOTE: We're including an entry for all vertices, even vertices that aren't referenced by any triangles (due to our
+				//       sparse array optimization.)
+				const int32 VertexArraySize = GetVertexArraySize();
+				VertexPositions.SetNum( VertexArraySize, false );
+				for( int32 VertexNumber = 0; VertexNumber < VertexArraySize; ++VertexNumber )
+				{
+					const FVertexID VertexID( VertexNumber );
+					if( IsValidVertex( VertexID ) )
+					{
+						VertexPositions[ VertexNumber ] = GetVertexAttribute( VertexID, UEditableMeshAttribute::VertexPosition(), 0 );
+					}
+					else
+					{
+						// Vertex isn't used, but we'll include a zero'd entry so that our indices still match up.
+						VertexPositions[ VertexNumber ] = FVector::ZeroVector;
+					}
+				}
+			}
+
+			// Texture coordinates (per polygon vertex)
+			FirstPolygonNumberForSections.Reset();
+
+			{
+				FVarVertexDatas.Reset();
+				FVarVertexDatas.Reserve( OsdFVarIndicesPerFace.Num() );
+
+				const int32 SectionArraySize = GetSectionArraySize();
+
+				int32 NumPolygonsSoFar = 0;
+				for( int32 SectionNumber = 0; SectionNumber < SectionArraySize; ++SectionNumber )
+				{
+					const FSectionID SectionID( SectionNumber );
+					if( IsValidSection( SectionID ) )
+					{
+						FirstPolygonNumberForSections.Add( NumPolygonsSoFar );
+						NumPolygonsSoFar += GetPolygonCount( SectionID );
+
+						int32 SectionPolygonArraySize = GetPolygonArraySize( SectionID );
+						for( int32 SectionPolygonNumber = 0; SectionPolygonNumber < SectionPolygonArraySize; ++SectionPolygonNumber )
+						{
+							const FPolygonRef PolygonRef( SectionID, FPolygonID( SectionPolygonNumber ) );
+							if( IsValidPolygon( PolygonRef ) )
+							{
+								const int32 PerimeterVertexCount = GetPolygonPerimeterVertexCount( PolygonRef );
+								for( int32 PerimeterVertexNumber = 0; PerimeterVertexNumber < PerimeterVertexCount; ++PerimeterVertexNumber )
+								{
+									FOsdFVarVertexData& FVarVertexData = *new( FVarVertexDatas ) FOsdFVarVertexData();
+									FVarVertexData.TextureCoordinates[ 0 ] = TextureCoordinateCount > 0 ? FVector2D( GetPolygonPerimeterVertexAttribute( PolygonRef, PerimeterVertexNumber, UEditableMeshAttribute::VertexTextureCoordinate(), 0 ) ) : FVector2D::ZeroVector;
+									FVarVertexData.TextureCoordinates[ 1 ] = TextureCoordinateCount > 1 ? FVector2D( GetPolygonPerimeterVertexAttribute( PolygonRef, PerimeterVertexNumber, UEditableMeshAttribute::VertexTextureCoordinate(), 1 ) ) : FVector2D::ZeroVector;
+									FVarVertexData.VertexColor = FLinearColor( GetPolygonPerimeterVertexAttribute( PolygonRef, PerimeterVertexNumber, UEditableMeshAttribute::VertexColor(), 0 ) );
+								}
+							}
+						}
+					}
+				}
+
+				check( FVarVertexDatas.Num() == OsdVertexIndicesPerFace.Num() );
+			}
+		}
+
+		// @todo mesheditor subdiv debug
+// 		for( int32 VertexIndex = 0; VertexIndex < FVarVertexDatas.Num(); ++VertexIndex )
+// 		{
+// 			const FOsdFVarVertexData& FVarVertexData = FVarVertexDatas[ VertexIndex ];
+// 			GWarn->Logf( TEXT( "FVar%i, U:%0.2f, V:%0.2f" ), VertexIndex, FVarVertexData.TextureCoordinates[0].X, FVarVertexData.TextureCoordinates[0].Y );
+// 		}
+
+		static TArray<FOsdFVarVertexData> LimitFVarVertexDatas;
+		LimitFVarVertexDatas.Reset();
+
+
+		// Start with the base cage geometry, and refine the geometry until we get to the limit surface
+		{
+			// NOTE: The OsdVertexPositions list might contains vertices that aren't actually referenced by any polygons (due to our
+			//       sparse array optimization.)  That's OK though.
+			{
+				int32 NextScratchBufferIndex = 0;
+				static TArray<FVector> ScratchVertexPositions[ 2 ];
+				ScratchVertexPositions[ 0 ].Reset();
+				ScratchVertexPositions[ 1 ].Reset();
+
+				for( int32 RefinementLevel = 1; RefinementLevel <= SubdivisionCount; ++RefinementLevel )
+				{
+					const OpenSubdiv::Far::TopologyLevel& OsdLevel = OsdTopologyRefiner->GetLevel( RefinementLevel );
+
+					// For the last refinement level, we'll copy positions straight to our output buffer (to avoid having to copy the data later.)
+					// For earlier levels, we'll ping-pong between scratch buffers.
+					const TArray<FVector>& SourceVertexPositions = ( RefinementLevel == 1 ) ? VertexPositions : ScratchVertexPositions[ !NextScratchBufferIndex ];
+					check( SourceVertexPositions.Num() == OsdTopologyRefiner->GetLevel( RefinementLevel - 1 ).GetNumVertices() );
+
+					TArray<FVector>& DestVertexPositions = ScratchVertexPositions[ NextScratchBufferIndex ];
+					DestVertexPositions.SetNumUninitialized( OsdLevel.GetNumVertices(), false );
+
+					FOsdVector* OsdDestVertexPositions = reinterpret_cast<FOsdVector*>( DestVertexPositions.GetData() );
+					OsdPrimvarRefiner.Interpolate( 
+						RefinementLevel, 
+						reinterpret_cast<const FOsdVector*>( SourceVertexPositions.GetData() ),
+						OsdDestVertexPositions );
+
+					NextScratchBufferIndex = !NextScratchBufferIndex;
+				}
+
+				// @todo mesheditor subdiv perf: We should probably be using stencils for faster updates (unless topology is changing every frame too)
+
+				// We've generated interpolated positions for the most fine subdivision level, but now we need to compute the positions
+				// on the limit surface.  While doing this, we also compute gradients at every vertex for either surface axis
+				{
+					const FOsdVector* OsdSourceLimitPositions = reinterpret_cast<FOsdVector*>( ScratchVertexPositions[ !NextScratchBufferIndex ].GetData() );
+					SubdivisionLimitData.VertexPositions.SetNumUninitialized( LimitVertexCount );
+					FOsdVector* OsdDestLimitPositions = reinterpret_cast<FOsdVector*>( SubdivisionLimitData.VertexPositions.GetData() );
+
+					LimitXGradients.SetNumUninitialized( LimitVertexCount );
+					FOsdVector* OsdDestLimitXGradients = reinterpret_cast<FOsdVector*>( LimitXGradients.GetData() );
+
+					LimitYGradients.SetNumUninitialized( LimitVertexCount );
+					FOsdVector* OsdDestLimitYGradients = reinterpret_cast<FOsdVector*>( LimitYGradients.GetData() );
+
+					OsdPrimvarRefiner.Limit( OsdSourceLimitPositions, /* Out */ OsdDestLimitPositions, /* Out */ OsdDestLimitXGradients, /* Out */ OsdDestLimitYGradients );
+
+					if( EditableMesh::InterpolatePositionsToLimit->GetInt() == 0 )
+					{
+						SubdivisionLimitData.VertexPositions = ScratchVertexPositions[ !NextScratchBufferIndex ];
+					}
+				}
+
+				check( LimitVertexCount == SubdivisionLimitData.VertexPositions.Num() );
+			}
+
+			{
+				const int32 FVarChannelNumber = 0;
+
+				static TArray<FOsdFVarVertexData> ScratchFVarVertexDatas[ 2 ];
+				ScratchFVarVertexDatas[ 0 ].Reset();
+				ScratchFVarVertexDatas[ 1 ].Reset();
+
+				int32 NextScratchBufferIndex = 0;
+				for( int32 RefinementLevel = 1; RefinementLevel <= SubdivisionCount; ++RefinementLevel )
+				{
+					const OpenSubdiv::Far::TopologyLevel& OsdLevel = OsdTopologyRefiner->GetLevel( RefinementLevel );
+
+					const TArray<FOsdFVarVertexData>& SourceFVarVertexDatas = ( RefinementLevel == 1 ) ? FVarVertexDatas : ScratchFVarVertexDatas[ !NextScratchBufferIndex ];
+					check( SourceFVarVertexDatas.Num() == OsdTopologyRefiner->GetLevel( RefinementLevel - 1 ).GetNumFVarValues( FVarChannelNumber ) );
+
+					TArray<FOsdFVarVertexData>& DestFVarVertexDatas = ScratchFVarVertexDatas[ NextScratchBufferIndex ];
+					DestFVarVertexDatas.SetNumUninitialized( OsdLevel.GetNumFVarValues( FVarChannelNumber ), false );
+
+					FOsdFVarVertexData* DestFVarVertexDatasPtr = DestFVarVertexDatas.GetData();
+					OsdPrimvarRefiner.InterpolateFaceVarying(
+						RefinementLevel,
+						SourceFVarVertexDatas.GetData(),
+						DestFVarVertexDatasPtr,
+						FVarChannelNumber );
+
+					NextScratchBufferIndex = !NextScratchBufferIndex;
+				}
+	
+				if( EditableMesh::InterpolateFVarsToLimit->GetInt() != 0 )
+				{
+					LimitFVarVertexDatas.SetNumUninitialized( OsdLimitLevel.GetNumFVarValues( FVarChannelNumber ), false );
+					FOsdFVarVertexData* DestFVarVertexDatasPtr = LimitFVarVertexDatas.GetData();
+					OsdPrimvarRefiner.LimitFaceVarying(	// @todo mesheditor subdiv fvars: The OSD tutorials don't bother doing a Limit pass for UVs/Colors...
+						ScratchFVarVertexDatas[ !NextScratchBufferIndex ].GetData(),
+						DestFVarVertexDatasPtr,
+						FVarChannelNumber );
+				}
+				else
+				{
+					LimitFVarVertexDatas = ScratchFVarVertexDatas[ !NextScratchBufferIndex ];
+				}
+			}
+		}
+
+		SubdivisionLimitData.Sections.SetNum( SectionCount, false );
+
+		for( int32 LimitFaceNumber = 0; LimitFaceNumber < LimitFaceCount; ++LimitFaceNumber )
+		{
+			const OpenSubdiv::Far::ConstIndexArray& OsdLimitFaceVertices = OsdLimitLevel.GetFaceVertices( LimitFaceNumber );
+			const int32 FaceVertexCount = OsdLimitFaceVertices.size();
+			check( FaceVertexCount == 4 );	// We're always expecting quads as the result of a Catmull-Clark subdivision
+
+			// Find the parent face in our original control mesh for this subdivided quad.  We'll use this to
+			// determine which section the face belongs to
+			// @todo mesheditor subdiv perf: We can use InterpolateFaceUniform() to push section indices down to the subdivided faces. More memory, might be slower (copies), but avoids iterating here.
+			int32 QuadSectionNumber = 0;
+			{
+				int32 CurrentFaceNumber = LimitFaceNumber;
+				for( int32 SubdivisionLevel = SubdivisionCount; SubdivisionLevel > 0; --SubdivisionLevel )
+				{
+					const OpenSubdiv::Far::TopologyLevel& OsdLevel = OsdTopologyRefiner->GetLevel( SubdivisionLevel );
+					const int32 ParentLevelFace = OsdLevel.GetFaceParentFace( CurrentFaceNumber );
+					CurrentFaceNumber = ParentLevelFace;
+				}
+				int32 BaseCageFaceNumber = CurrentFaceNumber;
+
+				for( int32 SectionNumber = SectionCount - 1; SectionNumber >= 0; --SectionNumber )
+				{
+					if( BaseCageFaceNumber >= FirstPolygonNumberForSections[ SectionNumber ] )
+					{
+						QuadSectionNumber = SectionNumber;
+						break;
+					}
+				}
+			}
+
+			const int32 FVarChannelNumber = 0;
+			const OpenSubdiv::Far::ConstIndexArray& OsdLimitFaceFVarValues = OsdLimitLevel.GetFaceFVarValues( LimitFaceNumber, FVarChannelNumber );
+			check( OsdLimitFaceFVarValues.size() == 4 );	// Expecting quads
+
+			FSubdivisionLimitSection& SubdivisionSection = SubdivisionLimitData.Sections[ QuadSectionNumber ];
+			FSubdividedQuad& SubdividedQuad = *new( SubdivisionSection.SubdividedQuads ) FSubdividedQuad();
+			for( int32 FaceVertexNumber = 0; FaceVertexNumber < FaceVertexCount; ++FaceVertexNumber )
+			{
+				FSubdividedQuadVertex& QuadVertex = SubdividedQuad.AccessQuadVertex( FaceVertexNumber );
+
+				QuadVertex.VertexPositionIndex = OsdLimitFaceVertices[ FaceVertexNumber ];
+
+				const int OsdLimitFaceFVarIndex = OsdLimitFaceFVarValues[ FaceVertexNumber ];
+				const FOsdFVarVertexData& FVarVertexData = LimitFVarVertexDatas[ OsdLimitFaceFVarIndex ];
+
+				QuadVertex.TextureCoordinate0 = FVarVertexData.TextureCoordinates[ 0 ];
+				QuadVertex.TextureCoordinate1 = FVarVertexData.TextureCoordinates[ 1 ];
+
+				QuadVertex.VertexColor = FVarVertexData.VertexColor.ToFColor( true );
+
+				QuadVertex.VertexNormal = FVector::CrossProduct( LimitXGradients[ QuadVertex.VertexPositionIndex ].GetSafeNormal(), LimitYGradients[ QuadVertex.VertexPositionIndex ].GetSafeNormal() );
+				
+				// NOTE: Tangents will be computed separately, below
+			}
+		}
+
+		// Compute normal and tangent vectors for each quad vertex, taking into account the texture coordinates
+		for( int32 SectionNumber = 0; SectionNumber < SubdivisionLimitData.Sections.Num(); ++SectionNumber )
+		{
+			struct FMikkUserData
+			{
+				FSubdivisionLimitData* LimitData;
+				int32 SectionNumber;
+
+				FMikkUserData( FSubdivisionLimitData* InitLimitData, const int32 InitSectionNumber )
+					: LimitData( InitLimitData ),
+						SectionNumber( InitSectionNumber )
+				{
+				}
+			} MikkUserData( &SubdivisionLimitData, SectionNumber );
+
+			struct Local
+			{
+				static int MikkGetNumFaces( const SMikkTSpaceContext* Context )
+				{
+					const FMikkUserData& UserData = *static_cast<FMikkUserData*>( Context->m_pUserData );
+					return UserData.LimitData->Sections[ UserData.SectionNumber ].SubdividedQuads.Num();
+				}
+
+				static int MikkGetNumVertsOfFace( const SMikkTSpaceContext* Context, const int MikkFaceIndex )
+				{
+					// Always quads
+					return 4;
+				}
+
+				static void MikkGetPosition( const SMikkTSpaceContext* Context, float OutPosition[ 3 ], const int MikkFaceIndex, const int MikkVertexIndex )
+				{
+					const FMikkUserData& UserData = *static_cast<FMikkUserData*>( Context->m_pUserData );
+					const FSubdividedQuadVertex& QuadVertex = UserData.LimitData->Sections[ UserData.SectionNumber ].SubdividedQuads[ MikkFaceIndex ].GetQuadVertex( MikkVertexIndex );
+					const FVector VertexPosition = UserData.LimitData->VertexPositions[ QuadVertex.VertexPositionIndex ];
+
+					OutPosition[ 0 ] = VertexPosition.X;
+					OutPosition[ 1 ] = VertexPosition.Y;
+					OutPosition[ 2 ] = VertexPosition.Z;
+				}
+
+				static void MikkGetNormal( const SMikkTSpaceContext* Context, float OutNormal[ 3 ], const int MikkFaceIndex, const int MikkVertexIndex )
+				{
+					const FMikkUserData& UserData = *static_cast<FMikkUserData*>( Context->m_pUserData );
+					const FSubdividedQuadVertex& QuadVertex = UserData.LimitData->Sections[ UserData.SectionNumber ].SubdividedQuads[ MikkFaceIndex ].GetQuadVertex( MikkVertexIndex );
+
+					OutNormal[ 0 ] = QuadVertex.VertexNormal.X;
+					OutNormal[ 1 ] = QuadVertex.VertexNormal.Y;
+					OutNormal[ 2 ] = QuadVertex.VertexNormal.Z;
+				}
+
+				static void MikkGetTexCoord( const SMikkTSpaceContext* Context, float OutUV[ 2 ], const int MikkFaceIndex, const int MikkVertexIndex )
+				{
+					const FMikkUserData& UserData = *static_cast<FMikkUserData*>( Context->m_pUserData );
+					const FSubdividedQuadVertex& QuadVertex = UserData.LimitData->Sections[ UserData.SectionNumber ].SubdividedQuads[ MikkFaceIndex ].GetQuadVertex( MikkVertexIndex );
+
+					// @todo mesheditor: Support using a custom texture coordinate index for tangent space generation?
+					OutUV[ 0 ] = QuadVertex.TextureCoordinate0.X;
+					OutUV[ 1 ] = QuadVertex.TextureCoordinate0.Y;
+				}
+
+				static void MikkSetTSpaceBasic( const SMikkTSpaceContext* Context, const float Tangent[ 3 ], const float BitangentSign, const int MikkFaceIndex, const int MikkVertexIndex )
+				{
+					FMikkUserData& UserData = *static_cast<FMikkUserData*>( Context->m_pUserData );
+					FSubdividedQuadVertex& QuadVertex = UserData.LimitData->Sections[ UserData.SectionNumber ].SubdividedQuads[ MikkFaceIndex ].AccessQuadVertex( MikkVertexIndex );
+
+					QuadVertex.VertexTangent = FVector( Tangent[ 0 ], Tangent[ 1 ], Tangent[ 2 ] );
+					QuadVertex.VertexBinormalSign = BitangentSign;
+				}
+			};
+
+			SMikkTSpaceInterface MikkTInterface;
+			{
+				MikkTInterface.m_getNumFaces = Local::MikkGetNumFaces;
+				MikkTInterface.m_getNumVerticesOfFace = Local::MikkGetNumVertsOfFace;
+				MikkTInterface.m_getPosition = Local::MikkGetPosition;
+				MikkTInterface.m_getNormal = Local::MikkGetNormal;
+				MikkTInterface.m_getTexCoord = Local::MikkGetTexCoord;
+
+				MikkTInterface.m_setTSpaceBasic = Local::MikkSetTSpaceBasic;
+				MikkTInterface.m_setTSpace = nullptr;
+			}
+
+			SMikkTSpaceContext MikkTContext;
+			{
+				MikkTContext.m_pInterface = &MikkTInterface;
+				MikkTContext.m_pUserData = (void*)( &MikkUserData );
+				MikkTContext.m_bIgnoreDegenerates = true;	// No degenerates in our list
+			}
+
+			// Now we'll ask MikkTSpace to actually generate the tangents
+			genTangSpaceDefault( &MikkTContext );
+		}
+
+		{
+			const int32 LimitEdgeCount = OsdLimitLevel.GetNumEdges();
+			for( int32 LimitEdgeNumber = 0; LimitEdgeNumber < LimitEdgeCount; ++LimitEdgeNumber )
+			{
+				const OpenSubdiv::Far::ConstIndexArray& OsdLimitEdgeVertices = OsdLimitLevel.GetEdgeVertices( LimitEdgeNumber );
+				const int32 EdgeVertexCount = OsdLimitEdgeVertices.size();
+				check( EdgeVertexCount == 2 );	// Edges always connect two vertices
+
+				FSubdividedWireEdge& SubdividedWireEdge = *new( SubdivisionLimitData.SubdividedWireEdges ) FSubdividedWireEdge();
+
+				SubdividedWireEdge.EdgeVertex0PositionIndex = OsdLimitEdgeVertices[ 0 ];
+				SubdividedWireEdge.EdgeVertex1PositionIndex = OsdLimitEdgeVertices[ 1 ];
+
+				// Default to not highlighting this edge as a base cage counterpart.  We'll actually figure this out below.
+				SubdividedWireEdge.bIsBaseCageCounterpartEdge = false;
+			}
+
+
+
+			{
+				// Figure out all of the unique polygon edges that make up the base cage
+				static TArray<int32> BaseCageEdges;
+				BaseCageEdges.Reset();
+				{
+					const OpenSubdiv::Far::TopologyLevel& OsdBaseCageLevel = OsdTopologyRefiner->GetLevel( 0 );
+					const int32 BaseCageFaceCount = OsdBaseCageLevel.GetNumFaces();
+					for( int32 BaseCageFaceNumber = 0; BaseCageFaceNumber < BaseCageFaceCount; ++BaseCageFaceNumber )
+					{
+						const OpenSubdiv::Far::ConstIndexArray& OsdBaseCageFaceEdges = OsdBaseCageLevel.GetFaceEdges( BaseCageFaceNumber );
+						for( int32 FaceEdgeNumber = 0; FaceEdgeNumber < OsdBaseCageFaceEdges.size(); ++FaceEdgeNumber )
+						{
+							BaseCageEdges.AddUnique( OsdBaseCageFaceEdges[ FaceEdgeNumber ] );	// @todo mesheditor now: Any better way to get all of the edges in base cage and avoid AddUnique 
+						}
+					}
+				}
+
+				// Go through and determine the limit child edges of all of the original base cage edges by drilling down through
+				// the subdivision heirarchy
+				int32 NextScratchBufferIndex = 0;
+				static TArray<int32> ScratchChildEdges[ 2 ];
+				ScratchChildEdges[ 0 ].Reset();
+				ScratchChildEdges[ 1 ].Reset();
+				{
+					for( int32 RefinementLevel = 0; RefinementLevel < SubdivisionCount; ++RefinementLevel )
+					{
+						const OpenSubdiv::Far::TopologyLevel& OsdLevel = OsdTopologyRefiner->GetLevel( RefinementLevel );
+
+						const TArray<int32>& SourceChildEdges = ( RefinementLevel == 0 ) ? BaseCageEdges : ScratchChildEdges[ !NextScratchBufferIndex ];
+
+						TArray<int32>& DestChildEdges = ScratchChildEdges[ NextScratchBufferIndex ];
+						DestChildEdges.Reset();
+
+						for( int32 SourceEdgeNumber = 0; SourceEdgeNumber < SourceChildEdges.Num(); ++SourceEdgeNumber )
+						{
+							const OpenSubdiv::Far::ConstIndexArray& OsdChildEdges = OsdLevel.GetEdgeChildEdges( SourceChildEdges[ SourceEdgeNumber ] );
+							for( int32 ChildEdgeNumber = 0; ChildEdgeNumber < OsdChildEdges.size(); ++ChildEdgeNumber )
+							{
+								DestChildEdges.Add( OsdChildEdges[ ChildEdgeNumber ] );
+							}
+						}
+
+						NextScratchBufferIndex = !NextScratchBufferIndex;
+					}
+				}
+
+				// No go back and update our subdivided wire edges, marking the edges that we determined were decendants of the base cage edges
+				const TArray<int32>& BaseCageCounterpartEdgesAtLimit = ScratchChildEdges[ !NextScratchBufferIndex ];
+				for( const int32 CounterpartEdgeAtLimit : BaseCageCounterpartEdgesAtLimit )
+				{
+					check( CounterpartEdgeAtLimit < SubdivisionLimitData.SubdividedWireEdges.Num() );
+					SubdivisionLimitData.SubdividedWireEdges[ CounterpartEdgeAtLimit ].bIsBaseCageCounterpartEdge = true;
+				}
+			}
+
+		}
+	}
+}
+
+
+void UEditableMesh::ComputePolygonTriangulation( const FPolygonRef PolygonRef, TArray<int32>& OutPerimeterVertexNumbersForTriangles ) const
+{
+	// NOTE: This polygon triangulation code is partially based on the ear cutting algorithm described on
+	//       page 497 of the book "Real-time Collision Detection", published in 2005.
+
+	struct Local
+	{
+		// Returns true if the triangle formed by the specified three positions has a normal that is facing the opposite direction of the reference normal
+		static inline bool IsTriangleFlipped( const FVector ReferenceNormal, const FVector VertexPositionA, const FVector VertexPositionB, const FVector VertexPositionC )
+		{
+			const FVector TriangleNormal = FVector::CrossProduct(
+				VertexPositionC - VertexPositionA,
+				VertexPositionB - VertexPositionA ).GetSafeNormal();
+			return ( FVector::DotProduct( ReferenceNormal, TriangleNormal ) <= 0.0f );
+		}
+
+	};
+
+
+	OutPerimeterVertexNumbersForTriangles.Reset();
+
+	// @todo mesheditor holes: Does not support triangles with holes yet!
+
+	// Polygon must have at least three vertices/edges
+	static TArray<FVertexID> VertexIDs;
+	GetPolygonPerimeterVertices( PolygonRef, /* Out */ VertexIDs );
+	const int32 PolygonVertexCount = VertexIDs.Num();
+	check( PolygonVertexCount >= 3 );
+
+	// First figure out the polygon normal.  We need this to determine which triangles are convex, so that
+	// we can figure out which ears to clip
+	const FVector PolygonNormal = ComputePolygonNormal( PolygonRef );
+
+	// Make a simple linked list array of the previous and next vertex numbers, for each vertex number
+	// in the polygon.  This will just save us having to iterate later on.
+	static TArray<int32> PrevVertexNumbers, NextVertexNumbers;
+	static TArray<FVector> VertexPositions;
+	{
+		PrevVertexNumbers.SetNumUninitialized( PolygonVertexCount, false );
+		NextVertexNumbers.SetNumUninitialized( PolygonVertexCount, false );
+		VertexPositions.SetNumUninitialized( PolygonVertexCount, false );
+
+		for( int32 VertexNumber = 0; VertexNumber < PolygonVertexCount; ++VertexNumber )
+		{
+			PrevVertexNumbers[ VertexNumber ] = VertexNumber - 1;
+			NextVertexNumbers[ VertexNumber ] = VertexNumber + 1;
+
+			VertexPositions[ VertexNumber ] = GetVertexAttribute( VertexIDs[ VertexNumber ], UEditableMeshAttribute::VertexPosition(), 0 );
+		}
+		PrevVertexNumbers[ 0 ] = PolygonVertexCount - 1;
+		NextVertexNumbers[ PolygonVertexCount - 1 ] = 0;
+	}
+
+	int32 EarVertexNumber = 0;
+	int32 EarTestCount = 0;
+	for( int32 RemainingVertexCount = PolygonVertexCount; RemainingVertexCount >= 3; )
+	{
+		bool bIsEar = true;
+
+		// If we're down to only a triangle, just treat it as an ear.  Also, if we've tried every possible candidate
+		// vertex looking for an ear, go ahead and just treat the current vertex as an ear.  This can happen when 
+		// vertices are colinear or other degenerate cases.
+		if( RemainingVertexCount > 3 && EarTestCount < RemainingVertexCount )
+		{
+			const FVector PrevVertexPosition = VertexPositions[ PrevVertexNumbers[ EarVertexNumber ] ];
+			const FVector EarVertexPosition = VertexPositions[ EarVertexNumber ];
+			const FVector NextVertexPosition = VertexPositions[ NextVertexNumbers[ EarVertexNumber ] ];
+
+			// Figure out whether the potential ear triangle is facing the same direction as the polygon
+			// itself.  If it's facing the opposite direction, then we're dealing with a concave triangle
+			// and we'll skip it for now.
+			if( !Local::IsTriangleFlipped(
+					PolygonNormal,
+					PrevVertexPosition,
+					EarVertexPosition,
+					NextVertexPosition ) )
+			{
+				int32 TestVertexNumber = NextVertexNumbers[ NextVertexNumbers[ EarVertexNumber ] ];
+
+				do 
+				{
+					// Test every other remaining vertex to make sure that it doesn't lie inside our potential ear
+					// triangle.  If we find a vertex that's inside the triangle, then it cannot actually be an ear.
+					const FVector TestVertexPosition = VertexPositions[ TestVertexNumber ];
+					if( FGeomTools::PointInTriangle(
+							PrevVertexPosition,
+							EarVertexPosition,
+							NextVertexPosition,
+							TestVertexPosition,
+							SMALL_NUMBER ) )
+					{
+						bIsEar = false;
+						break;
+					}
+
+					TestVertexNumber = NextVertexNumbers[ TestVertexNumber ];
+				} 
+				while( TestVertexNumber != PrevVertexNumbers[ EarVertexNumber ] );
+			}
+			else
+			{
+				bIsEar = false;
+			}
+		}
+
+		if( bIsEar )
+		{
+			// OK, we found an ear!  Let's save this triangle in our output buffer.
+			{
+				const int32 TriangleVertexIDNumbers[ 3 ] = { PrevVertexNumbers[ EarVertexNumber ], EarVertexNumber, NextVertexNumbers[ EarVertexNumber ] };
+				for( int32 TriangleVertexNumber = 0; TriangleVertexNumber < 3; ++TriangleVertexNumber )
+				{
+					OutPerimeterVertexNumbersForTriangles.Add( TriangleVertexIDNumbers[ TriangleVertexNumber ] );
+				}
+			}
+
+			// Update our linked list.  We're effectively cutting off the ear by pointing the ear vertex's neighbors to
+			// point at their next sequential neighbor, and reducing the remaining vertex count by one.
+			{
+				NextVertexNumbers[ PrevVertexNumbers[ EarVertexNumber ] ] = NextVertexNumbers[ EarVertexNumber ];
+				PrevVertexNumbers[ NextVertexNumbers[ EarVertexNumber ] ] = PrevVertexNumbers[ EarVertexNumber ];
+				--RemainingVertexCount;
+			}
+
+			// Move on to the previous vertex in the list, now that this vertex was cut
+			EarVertexNumber = PrevVertexNumbers[ EarVertexNumber ];
+
+			EarTestCount = 0;
+		}
+		else
+		{
+			// The vertex is not the ear vertex, because it formed a triangle that either had a normal which pointed in the opposite direction
+			// of the polygon, or at least one of the other polygon vertices was found to be inside the triangle.  Move on to the next vertex.
+			EarVertexNumber = NextVertexNumbers[ EarVertexNumber ];
+
+			// Keep track of how many ear vertices we've tested, so that if we exhaust all remaining vertices, we can
+			// fall back to clipping the triangle and adding it to our mesh anyway.  This is important for degenerate cases.
+			++EarTestCount;
+		}
+	}
+
+	check( OutPerimeterVertexNumbersForTriangles.Num() > 0 );
+	check( OutPerimeterVertexNumbersForTriangles.Num() % 3 == 0 );	// Must all be triangles
+}
+
+
+void UEditableMesh::ComputePolygonTriangulation( const FPolygonRef PolygonRef, TArray<FVertexID>& OutTrianglesVertexIDs ) const
+{
+	static TArray<int32> PerimeterVertexNumbersForTriangles;
+	ComputePolygonTriangulation( PolygonRef, /* Out */ PerimeterVertexNumbersForTriangles );
+
+	static TArray<FVertexID> PerimeterVertexIDs;
+	GetPolygonPerimeterVertices( PolygonRef, /* Out */ PerimeterVertexIDs );
+
+	OutTrianglesVertexIDs.SetNum( PerimeterVertexNumbersForTriangles.Num(), false );
+	for( int32 TriangleVerticesNumber = 0; TriangleVerticesNumber < PerimeterVertexNumbersForTriangles.Num(); ++TriangleVerticesNumber )
+	{
+		const int32 PerimeterVertexNumber = PerimeterVertexNumbersForTriangles[ TriangleVerticesNumber ];
+		OutTrianglesVertexIDs[ TriangleVerticesNumber ] = PerimeterVertexIDs[ PerimeterVertexNumber ];
+	}
+}
+
+bool UEditableMesh::ComputeBarycentricWeightForPointOnPolygon( const FPolygonRef PolygonRef, const FVector PointOnPolygon, TArray<int32>& OutPerimeterVertexIndices, FVector& OutTriangleVertexWeights ) const
+{
+	OutPerimeterVertexIndices.Reset( 3 );
+
+	// Triangulate the polygon, so that we can compute barycentric weights for the texture coordinates that will be used
+	// for the new inset vertices
+	static TArray<int32> PerimeterVertexNumbersForTriangles;
+	ComputePolygonTriangulation( PolygonRef, /* Out */ PerimeterVertexNumbersForTriangles );
+
+	check( PerimeterVertexNumbersForTriangles.Num() % 3 == 0 );
+	const int32 TriangleCount = PerimeterVertexNumbersForTriangles.Num() / 3;
+
+	// Figure out which triangle the incoming point is within
+	for( int32 TriangleNumber = 0; TriangleNumber < TriangleCount; ++TriangleNumber )
+	{
+		const FVector TriangleVertex0Position = GetPolygonPerimeterVertexAttribute( PolygonRef, PerimeterVertexNumbersForTriangles[ TriangleNumber * 3 + 0 ], UEditableMeshAttribute::VertexPosition(), 0 );
+		const FVector TriangleVertex1Position = GetPolygonPerimeterVertexAttribute( PolygonRef, PerimeterVertexNumbersForTriangles[ TriangleNumber * 3 + 1 ], UEditableMeshAttribute::VertexPosition(), 0 );
+		const FVector TriangleVertex2Position = GetPolygonPerimeterVertexAttribute( PolygonRef, PerimeterVertexNumbersForTriangles[ TriangleNumber * 3 + 2 ], UEditableMeshAttribute::VertexPosition(), 0 );
+
+		// Calculate the barycentric weights for the triangle's verts and determine if the point lies within its bounds.
+		OutTriangleVertexWeights = FMath::ComputeBaryCentric2D( PointOnPolygon, TriangleVertex0Position, TriangleVertex1Position, TriangleVertex2Position );
+
+		if( OutTriangleVertexWeights.X >= 0.0f && OutTriangleVertexWeights.Y >= 0.0f && OutTriangleVertexWeights.Z >= 0.0f )
+		{
+			// Okay, we found a triangle that the point is inside!  Return the perimeter vertex indices.
+			OutPerimeterVertexIndices.Add( PerimeterVertexNumbersForTriangles[ TriangleNumber * 3 + 0 ] );
+			OutPerimeterVertexIndices.Add( PerimeterVertexNumbersForTriangles[ TriangleNumber * 3 + 1 ] );
+			OutPerimeterVertexIndices.Add( PerimeterVertexNumbersForTriangles[ TriangleNumber * 3 + 2 ] );
+
+			return true;
+		}
+	}
+
+	return false;
+}
+
+
+void UEditableMesh::ComputeTextureCoordinatesForPointOnPolygon( const FPolygonRef PolygonRef, const FVector PointOnPolygon, bool& bOutWereTextureCoordinatesFound, TArray<FVector4>& OutInterpolatedTextureCoordinates )
+{
+	bOutWereTextureCoordinatesFound = false;
+	OutInterpolatedTextureCoordinates.Reset();
+
+	static TArray<int32> PerimeterVertexNumbers;
+	FVector TriangleVertexWeights;
+	if( ComputeBarycentricWeightForPointOnPolygon( PolygonRef, PointOnPolygon, PerimeterVertexNumbers, TriangleVertexWeights ) )
+	{
+		OutInterpolatedTextureCoordinates.SetNum( TextureCoordinateCount, false );
+
+		for( int32 TextureCoordinateIndex = 0; TextureCoordinateIndex < TextureCoordinateCount; ++TextureCoordinateIndex )
+		{
+			OutInterpolatedTextureCoordinates[ TextureCoordinateIndex ] =
+				TriangleVertexWeights.X * GetPolygonPerimeterVertexAttribute( PolygonRef, PerimeterVertexNumbers[ 0 ], UEditableMeshAttribute::VertexTextureCoordinate(), TextureCoordinateIndex ) +
+				TriangleVertexWeights.Y * GetPolygonPerimeterVertexAttribute( PolygonRef, PerimeterVertexNumbers[ 1 ], UEditableMeshAttribute::VertexTextureCoordinate(), TextureCoordinateIndex ) +
+				TriangleVertexWeights.Z * GetPolygonPerimeterVertexAttribute( PolygonRef, PerimeterVertexNumbers[ 2 ], UEditableMeshAttribute::VertexTextureCoordinate(), TextureCoordinateIndex );
+		}
+
+		bOutWereTextureCoordinatesFound = true;
+	}
+}
+
+
+void UEditableMesh::SetSubdivisionCount( const int32 NewSubdivisionCount )
+{
+	// @todo mesheditor subdiv: Really, instead of a custom FChange type for this type of change, we could create a change that
+	// represents a property (or set of properties) being changed, and use Unreal reflection to update it.  This would be
+	// reusable for future properties.
+
+	const bool bEnablingSubdivisionPreview = ( GetSubdivisionCount() == 0 && NewSubdivisionCount > 0 );
+	const bool bDisablingSubdivisionPreview = ( GetSubdivisionCount() > 0 && NewSubdivisionCount == 0 );
+
+	FSetSubdivisionCountChangeInput RevertInput;
+	RevertInput.NewSubdivisionCount = GetSubdivisionCount();
+
+	this->SubdivisionCount = NewSubdivisionCount;
+
+	if( bDisablingSubdivisionPreview )
+	{
+		// We've turned off subdivision preview, so we'll need to re-create the static mesh data from our stored mesh representation
+		RebuildRenderMesh();
+	}
+	else
+	{
+		// NOTE: We don't bother regenerating geometry here because it's expected that EndModification() will be called after this, which will do the trick
+	}
+
+	AddUndo( MakeUnique<FSetSubdivisionCountChange>( MoveTemp( RevertInput ) ) );
+}
+
+
+void UEditableMesh::MoveVertices( const TArray<FVertexToMove>& VerticesToMove )
+{
+	static TSet< FPolygonRef > PolygonsThatNeedNewNormals;
+	PolygonsThatNeedNewNormals.Reset();
+
+	static TArray<FAttributesForVertex> VertexAttributesToSet;
+	VertexAttributesToSet.Reset();
+
+	for( const FVertexToMove& VertexToMove : VerticesToMove )
+	{
+		const FVector CurrentPosition = GetVertexAttribute( VertexToMove.VertexID, UEditableMeshAttribute::VertexPosition(), 0 );
+
+		if( VertexToMove.NewVertexPosition != CurrentPosition )
+		{
+			FAttributesForVertex& AttributesForVertex = *new( VertexAttributesToSet ) FAttributesForVertex();
+			AttributesForVertex.VertexID = VertexToMove.VertexID;
+			FMeshElementAttributeData& VertexAttribute = *new( AttributesForVertex.VertexAttributes.Attributes ) FMeshElementAttributeData( UEditableMeshAttribute::VertexPosition(), 0, VertexToMove.NewVertexPosition );
+
+			// All of the polygons that share this vertex will need new normals
+			static TArray<FPolygonRef> ConnectedPolygonRefs;
+			GetVertexConnectedPolygons( VertexToMove.VertexID, /* Out */ ConnectedPolygonRefs );
+			PolygonsThatNeedNewNormals.Append( ConnectedPolygonRefs );
+		}
+	}
+
+	// Convert from a TSet to a TArray
+	static TArray<FPolygonRef> PolygonsToUpdate;
+	PolygonsToUpdate.Reset();
+	for( const FPolygonRef PolygonRef : PolygonsThatNeedNewNormals )
+	{
+		PolygonsToUpdate.Add( PolygonRef );
+	}
+
+	// Perform retriangulation after the geometry is moved back, when undoing only
+	{
+		const bool bOnlyOnUndo = true;
+		RetriangulatePolygons( PolygonsToUpdate, bOnlyOnUndo );
+	}
+
+	SetVerticesAttributes( VertexAttributesToSet );
+
+	// Generate new normals and tangents for any geometry that was moved around
+	GenerateNormalsAndTangentsForPolygonsAndAdjacents( PolygonsToUpdate );
+
+	// Everything needs to be retriangulated because convexity may have changed
+	// @todo mesheditor perf: Polygons will actually be triangulated twice in the same frame when performing an operation that implicitly retriangulates at the end of the op (once before move, once after move.)  Not super ideal.
+	{
+		const bool bOnlyOnUndo = false;
+		RetriangulatePolygons( PolygonsToUpdate, bOnlyOnUndo );
+	}
+}
+
+
+void UEditableMesh::CreateMissingPolygonPerimeterEdges( const FPolygonRef PolygonRef, TArray<FEdgeID>& OutNewEdgeIDs )
+{
+	OutNewEdgeIDs.Reset();
+
+	const int32 NumPolygonPerimeterEdges = GetPolygonPerimeterEdgeCount( PolygonRef );
+	const int32 NumPolygonPerimeterVertices = NumPolygonPerimeterEdges;		// Edge and vertex count are always the same
+
+	for( int32 PerimeterEdgeNumber = 0; PerimeterEdgeNumber < NumPolygonPerimeterEdges; ++PerimeterEdgeNumber )
+	{
+		const int32 PerimeterVertexNumber = PerimeterEdgeNumber;	// Edge and vertex counts are always the same
+
+		const FVertexID VertexID = GetPolygonPerimeterVertex( PolygonRef, PerimeterVertexNumber );
+		const FVertexID NextVertexID = GetPolygonPerimeterVertex( PolygonRef, ( PerimeterVertexNumber + 1 ) % NumPolygonPerimeterVertices );
+
+		// Find the edge that connects these vertices
+		FEdgeID FoundEdgeID = FEdgeID::Invalid;
+		bool bFoundEdge = false;
+
+		const int32 NumVertexConnectedEdges = GetVertexConnectedEdgeCount( VertexID );	// @todo mesheditor perf: Could be made faster by always iterating over the vertex with the fewest number of connected edges, instead of the first edge vertex
+		for( int32 VertexEdgeNumber = 0; VertexEdgeNumber < NumVertexConnectedEdges; ++VertexEdgeNumber )
+		{
+			const FEdgeID VertexConnectedEdgeID = GetVertexConnectedEdge( VertexID, VertexEdgeNumber );
+
+			// Try the edge's first vertex.  Does it point to our next edge?
+			FVertexID OtherEdgeVertexID = GetEdgeVertex( VertexConnectedEdgeID, 0 );
+			if( OtherEdgeVertexID == VertexID )
+			{
+				// Must be the the other one
+				OtherEdgeVertexID = GetEdgeVertex( VertexConnectedEdgeID, 1 );
+			}
+			else
+			{
+				check( GetEdgeVertex( VertexConnectedEdgeID, 1 ) == VertexID );
+			}
+
+			if( OtherEdgeVertexID == NextVertexID )
+			{
+				// We found the edge!
+				FoundEdgeID = VertexConnectedEdgeID;
+				bFoundEdge = true;
+				break;
+			}
+		}
+
+		if( !bFoundEdge )
+		{
+			// Create the new edge!  Note that this does not connect the edge to the polygon.  We expect the caller to do that afterwards.
+			FEdgeToCreate EdgeToCreate;
+			EdgeToCreate.VertexID0 = VertexID;
+			EdgeToCreate.VertexID1 = NextVertexID;
+
+			static TArray< FEdgeToCreate > EdgesToCreate;
+			EdgesToCreate.Reset();
+			EdgesToCreate.Add( EdgeToCreate );
+
+			static TArray< FEdgeID > NewEdgeIDs;
+			NewEdgeIDs.Reset();
+			CreateEdges( EdgesToCreate, /* Out */ NewEdgeIDs );
+
+			OutNewEdgeIDs.Append( NewEdgeIDs );
+		}
+	}
+}
+
+
+void UEditableMesh::CreateMissingPolygonHoleEdges( const FPolygonRef PolygonRef, const int32 HoleNumber, TArray<FEdgeID>& OutNewEdgeIDs )
+{
+	OutNewEdgeIDs.Reset();
+
+	const int32 NumHoleEdges = GetPolygonHoleEdgeCount( PolygonRef, HoleNumber );
+	const int32 NumHoleVertices = NumHoleEdges;		// Edge and vertex count are always the same
+
+	for( int32 HoleEdgeNumber = 0; HoleEdgeNumber < NumHoleEdges; ++HoleEdgeNumber )
+	{
+		const int32 HoleVertexNumber = HoleEdgeNumber;	// Edge and vertex counts are always the same
+
+		const FVertexID VertexID = GetPolygonHoleVertex( PolygonRef, HoleNumber, HoleVertexNumber );
+		const FVertexID NextVertexID = GetPolygonHoleVertex( PolygonRef, HoleNumber, ( HoleVertexNumber + 1 ) % NumHoleVertices );
+
+		// Find the edge that connects these vertices
+		FEdgeID FoundEdgeID = FEdgeID::Invalid;
+		{
+			bool bFoundEdge = false;
+
+			const int32 NumVertexConnectedEdges = GetVertexConnectedEdgeCount( VertexID );	// @todo mesheditor perf: Could be made faster by always iterating over the vertex with the fewest number of connected edges, instead of the first edge vertex
+			for( int32 VertexEdgeNumber = 0; VertexEdgeNumber < NumVertexConnectedEdges; ++VertexEdgeNumber )
+			{
+				const FEdgeID VertexConnectedEdgeID = GetVertexConnectedEdge( VertexID, VertexEdgeNumber );
+
+				// Try the edge's first vertex.  Does it point to us?
+				FVertexID OtherEdgeVertexID = GetEdgeVertex( VertexConnectedEdgeID, 0 );
+				if( OtherEdgeVertexID == VertexID )
+				{
+					// Must be the the other one
+					OtherEdgeVertexID = GetEdgeVertex( VertexConnectedEdgeID, 1 );
+				}
+
+				if( OtherEdgeVertexID == NextVertexID )
+				{
+					// We found the edge!
+					FoundEdgeID = VertexConnectedEdgeID;
+					bFoundEdge = true;
+					break;
+				}
+			}
+
+			if( !bFoundEdge )
+			{
+				// Create the new edge!  Note that this does not connect the edge to the polygon.  We expect the caller to do that afterwards.
+				FEdgeToCreate EdgeToCreate;
+				EdgeToCreate.VertexID0 = VertexID;
+				EdgeToCreate.VertexID1 = NextVertexID;
+
+				static TArray< FEdgeToCreate > EdgesToCreate;
+				EdgesToCreate.Reset();
+				EdgesToCreate.Add( EdgeToCreate );
+
+				static TArray< FEdgeID > NewEdgeIDs;
+				NewEdgeIDs.Reset();
+				CreateEdges( EdgesToCreate, /* Out */ NewEdgeIDs );
+
+				OutNewEdgeIDs.Append( NewEdgeIDs );
+			}
+		}
+	}
+}
+
+
+void UEditableMesh::SplitEdge( const FEdgeID EdgeID, const TArray<float>& Splits, TArray<FVertexID>& OutNewVertexIDs )
+{
+	// NOTE: The incoming splits should always be between 0.0 and 1.0, representing progress along 
+	//       the edge from the edge's first vertex toward it's other vertex.  The order doesn't matter (we'll sort them.)
+
+	check( Splits.Num() > 0 );
+
+	// Sort the split values smallest to largest.  We'll be adding a strip of vertices for each split, and the
+	// indices for those new vertices need to be in order.
+	static TArray<float> SortedSplits;
+	SortedSplits.Reset();
+	SortedSplits = Splits;
+	if( SortedSplits.Num() > 1 )
+	{
+		SortedSplits.Sort();
+	}
+
+	FVertexID OriginalEdgeVertexIDs[ 2 ];
+	GetEdgeVertices( EdgeID, /* Out */ OriginalEdgeVertexIDs[0], /* Out */ OriginalEdgeVertexIDs[1] );
+
+	static TArray<FVector> NewVertexPositions;
+	NewVertexPositions.Reset();
+	NewVertexPositions.AddUninitialized( SortedSplits.Num() );
+
+	for( int32 NewVertexNumber = 0; NewVertexNumber < NewVertexPositions.Num(); ++NewVertexNumber )
+	{
+		const float Split = SortedSplits[ NewVertexNumber ];
+		check( Split >= 0.0f && Split <= 1.0f );
+
+		NewVertexPositions[ NewVertexNumber ] = FMath::Lerp(
+			FVector( GetVertexAttribute( OriginalEdgeVertexIDs[ 0 ], UEditableMeshAttribute::VertexPosition(), 0 ) ),
+			FVector( GetVertexAttribute( OriginalEdgeVertexIDs[ 1 ], UEditableMeshAttribute::VertexPosition(), 0 ) ),
+			Split );
+	}
+
+	// Split the edge, and connect the vertex to the polygons that share the two new edges
+	const FVertexID OriginalEdgeFarVertexID = OriginalEdgeVertexIDs[ 1 ];
+
+	static TArray<FMeshElementAttributeData> OriginalEdgeAttributes;
+	OriginalEdgeAttributes.Reset();
+	{
+		for( const FName AttributeName : UEditableMesh::GetValidEdgeAttributes() )
+		{
+			const int32 MaxAttributeIndex = GetMaxAttributeIndex( AttributeName );
+			for( int32 AttributeIndex = 0; AttributeIndex < MaxAttributeIndex; ++AttributeIndex )
+			{
+				FMeshElementAttributeData& EdgeAttribute = *new( OriginalEdgeAttributes ) FMeshElementAttributeData(
+					AttributeName,
+					AttributeIndex,
+					GetEdgeAttribute( EdgeID, AttributeName, AttributeIndex ) );
+			}
+		}
+	}
+
+
+	// Create new vertices.  Their texture coordinates will be undefined, for now.  We'll set that later by using the connected 
+	// polygon's vertex data.  Also, the vertex's normals and tangents will be automatically generated.
+	static TArray<FVertexID> NewVertexIDs;
+	NewVertexIDs.Reset();
+	{
+		OutNewVertexIDs.Reset();
+		OutNewVertexIDs.Reserve( NewVertexPositions.Num() );
+
+		static TArray<FVertexToCreate> VerticesToCreate;
+		VerticesToCreate.Reset();
+		VerticesToCreate.Reserve( NewVertexPositions.Num() );
+
+		for( int32 NewVertexNumber = 0; NewVertexNumber < NewVertexPositions.Num(); ++NewVertexNumber )
+		{
+			const FVector& NewVertexPosition = NewVertexPositions[ NewVertexNumber ];
+
+			FVertexToCreate& VertexToCreate = *new( VerticesToCreate )FVertexToCreate();
+			VertexToCreate.VertexAttributes.Attributes.Add( FMeshElementAttributeData( UEditableMeshAttribute::VertexPosition(), 0, NewVertexPosition ) );
+		}
+
+		CreateVertices( VerticesToCreate, /* Out */ NewVertexIDs );
+
+		OutNewVertexIDs.Append( NewVertexIDs );
+	}
+
+	// Figure out which polygons are connected to the original edge.  Those polygons are going to need their
+	// connected vertices updated later on.
+	struct FAffectedPolygonEdge
+	{
+		FPolygonRef PolygonRef;
+		int32 PolygonVertexNumbers[ 2 ];
+	};
+	static TArray<FAffectedPolygonEdge> AffectedPolygonEdges;
+	AffectedPolygonEdges.Reset();
+
+	static TArray<FPolygonRef> AffectedPolygons;
+	AffectedPolygons.Reset();
+
+	{
+		const int32 ConnectedPolygonCount = GetEdgeConnectedPolygonCount( EdgeID );
+		for( int32 ConnectedPolygonNumber = 0; ConnectedPolygonNumber < ConnectedPolygonCount; ++ConnectedPolygonNumber )
+		{
+			const FPolygonRef PolygonRef = GetEdgeConnectedPolygon( EdgeID, ConnectedPolygonNumber );
+
+			// Figure out which of this polygon's vertices are part of the edge we're adding a vertex to
+			// @todo mesheditor hole: Add support for polygon holes, too!
+			static TArray<FVertexID> PolygonPerimeterVertexIDs;
+			PolygonPerimeterVertexIDs.Reset();
+			GetPolygonPerimeterVertices( PolygonRef, /* Out */ PolygonPerimeterVertexIDs );
+
+			int32 PolygonVertexNumberForEdgeVertices[ 2 ] = { INDEX_NONE, INDEX_NONE };
+			for( int32 PerimeterVertexNumber = 0; PerimeterVertexNumber < PolygonPerimeterVertexIDs.Num(); ++PerimeterVertexNumber )
+			{
+				const FVertexID PerimeterVertexID = PolygonPerimeterVertexIDs[ PerimeterVertexNumber ];
+				if( PerimeterVertexID == OriginalEdgeVertexIDs[ 0 ] )
+				{
+					// Okay, this polygon vertex is part of our edge (first edge vertex)
+					PolygonVertexNumberForEdgeVertices[ 0 ] = PerimeterVertexNumber;
+				}
+				else if( PerimeterVertexID == OriginalEdgeVertexIDs[ 1 ] )
+				{
+					// Okay, this polygon vertex is part of our edge (second edge vertex)
+					PolygonVertexNumberForEdgeVertices[ 1 ] = PerimeterVertexNumber;
+				}
+			}
+			check( PolygonVertexNumberForEdgeVertices[ 0 ] != INDEX_NONE && PolygonVertexNumberForEdgeVertices[ 1 ] != INDEX_NONE );
+
+			FAffectedPolygonEdge& AffectedPolygonEdge = *new( AffectedPolygonEdges ) FAffectedPolygonEdge();
+			AffectedPolygonEdge.PolygonRef = PolygonRef;
+			AffectedPolygonEdge.PolygonVertexNumbers[ 0 ] = PolygonVertexNumberForEdgeVertices[ 0 ];
+			AffectedPolygonEdge.PolygonVertexNumbers[ 1 ] = PolygonVertexNumberForEdgeVertices[ 1 ];
+
+			// @todo mesheditor perf: Not sure if we actually need to AddUnique() here -- can an edge be connected to the same polygon more than once? (tube shape). Seems sketchy if so?  And if so, GetEdgeConnectedPolygonCount() probably should have only reported 1 connection anyway
+			AffectedPolygons.AddUnique( AffectedPolygonEdge.PolygonRef );
+		}
+	}
+
+
+	// When undoing this change, we'll need to retriangulate before the vertices are destroyed (after CreateVertices() is called above),
+	// so that the vertices will be fully orphaned first (otherwise, the rendering vertices would still be referenced.)
+	{
+		const bool bOnlyOnUndo = true;
+		RetriangulatePolygons( AffectedPolygons, bOnlyOnUndo );
+	}
+
+
+	{
+		static TArray<FVerticesForEdge> VerticesForEdges;
+		VerticesForEdges.Reset();
+
+		// We'll keep the existing edge, but update it to connect to the first new vertex.  The second vertex of the edge will 
+		// now connect to the first (new) vertex ID, and so on.  The incoming vertices are expected to be ordered correctly.
+		{
+			FVerticesForEdge& VerticesForEdge = *new( VerticesForEdges ) FVerticesForEdge();
+
+			VerticesForEdge.EdgeID = EdgeID;
+			VerticesForEdge.NewVertexID0 = OriginalEdgeVertexIDs[ 0 ];
+			VerticesForEdge.NewVertexID1 = NewVertexIDs[ 0 ];
+		}
+
+		SetEdgesVertices( VerticesForEdges );
+	}
+
+	// Create new edges.  One for each of the new vertex positions passed in.
+	{
+		const int32 NewEdgeCount = NewVertexPositions.Num();
+
+		static TArray<FEdgeToCreate> EdgesToCreate;
+		EdgesToCreate.Reset();
+		EdgesToCreate.Reserve( NewEdgeCount );
+		for( int32 NewEdgeNumber = 0; NewEdgeNumber < NewEdgeCount; ++NewEdgeNumber )
+		{
+			FEdgeToCreate& EdgeToCreate = *new( EdgesToCreate ) FEdgeToCreate();
+
+			EdgeToCreate.VertexID0 = NewVertexIDs[ NewEdgeNumber ];
+			EdgeToCreate.VertexID1 = ( NewEdgeNumber == ( NewEdgeCount - 1 ) ) ? OriginalEdgeFarVertexID : NewVertexIDs[ NewEdgeNumber + 1 ];
+
+			EdgeToCreate.ConnectedPolygons = AffectedPolygons;
+
+			// Copy edge attributes over from original edge
+			EdgeToCreate.EdgeAttributes.Attributes = OriginalEdgeAttributes;
+		}
+
+		static TArray< FEdgeID > NewEdgeIDs;
+		NewEdgeIDs.Reset();
+		CreateEdges( EdgesToCreate, /* Out */ NewEdgeIDs );
+	}
+
+	// Update all affected polygons with their new vertices.  Also, we'll fill in polygon-specific vertex attributes (texture coordinates)
+	{
+		for( const FAffectedPolygonEdge& AffectedPolygonEdge : AffectedPolygonEdges )
+		{
+			const FPolygonRef PolygonRef = AffectedPolygonEdge.PolygonRef;
+
+			// Figure out where to start inserting vertices into this polygon, as well as which order to insert them.
+			bool bWindsForward = AffectedPolygonEdge.PolygonVertexNumbers[ 1 ] > AffectedPolygonEdge.PolygonVertexNumbers[ 0 ];
+			const int32 LargerIndex = bWindsForward ? AffectedPolygonEdge.PolygonVertexNumbers[ 1 ] : AffectedPolygonEdge.PolygonVertexNumbers[ 0 ];
+			const bool bIsContiguous = FMath::Abs( (int32)AffectedPolygonEdge.PolygonVertexNumbers[ 1 ] - (int32)AffectedPolygonEdge.PolygonVertexNumbers[ 0 ] ) == 1;
+			if( !bIsContiguous )
+			{
+				bWindsForward = !bWindsForward;
+			}
+
+			const int32 InsertAtIndex = bIsContiguous ? LargerIndex : ( LargerIndex + 1 );
+
+			// Figure out what the polygon's NEW vertex count (after we insert new vertices) will be
+			const int32 NewPolygonPerimeterVertexCount = GetPolygonPerimeterVertexCount( PolygonRef ) + NewVertexPositions.Num();
+
+			static TArray<FVertexAndAttributes> VerticesToInsert;
+			VerticesToInsert.Reset( NewVertexPositions.Num() );
+			VerticesToInsert.SetNum( NewVertexPositions.Num(), false );
+			for( int32 InsertVertexNumber = 0; InsertVertexNumber < NewVertexPositions.Num(); ++InsertVertexNumber )
+			{
+				FVertexAndAttributes& VertexToInsert = VerticesToInsert[ InsertVertexNumber ];
+
+				const int32 DirectionalVertexNumber = bWindsForward ? InsertVertexNumber : ( ( NewVertexPositions.Num() - 1 ) - InsertVertexNumber );
+				VertexToInsert.VertexID = NewVertexIDs[ InsertVertexNumber ];
+
+				// Generate texture coordinates for this vertex by interpolating between the polygon vertex texture coordinates on this polygon edge
+				{
+					const float Split = SortedSplits[ DirectionalVertexNumber ];
+					for( int32 TextureCoordinateIndex = 0; TextureCoordinateIndex < TextureCoordinateCount; ++TextureCoordinateIndex )
+					{
+						const FVector4 NewPolygonVertexTextureCoordinate = FMath::Lerp(
+							GetPolygonPerimeterVertexAttribute( PolygonRef, AffectedPolygonEdge.PolygonVertexNumbers[ 0 ], UEditableMeshAttribute::VertexTextureCoordinate(), TextureCoordinateIndex ),
+							GetPolygonPerimeterVertexAttribute( PolygonRef, AffectedPolygonEdge.PolygonVertexNumbers[ 1 ], UEditableMeshAttribute::VertexTextureCoordinate(), TextureCoordinateIndex ),
+							Split );
+
+						VertexToInsert.PolygonVertexAttributes.Attributes.Add( FMeshElementAttributeData( UEditableMeshAttribute::VertexTextureCoordinate(), TextureCoordinateIndex, NewPolygonVertexTextureCoordinate ) );
+					}
+				}
+
+				// Generate vertex color for this vertex by interpolating between the vertex colors on the edge
+				{
+					const float Split = SortedSplits[ DirectionalVertexNumber ];
+					const FVector4 NewVertexColor = FMath::Lerp(
+						GetPolygonPerimeterVertexAttribute( PolygonRef, AffectedPolygonEdge.PolygonVertexNumbers[ 0 ], UEditableMeshAttribute::VertexColor(), 0 ),
+						GetPolygonPerimeterVertexAttribute( PolygonRef, AffectedPolygonEdge.PolygonVertexNumbers[ 1 ], UEditableMeshAttribute::VertexColor(), 0 ),
+						Split );
+
+					VertexToInsert.PolygonVertexAttributes.Attributes.Add( FMeshElementAttributeData( UEditableMeshAttribute::VertexColor(), 0, NewVertexColor ) );
+				}
+			}
+
+			// Add the new vertices to the polygon, and set texture coordinates for all of the affected polygon vertices
+			InsertPolygonPerimeterVertices( PolygonRef, InsertAtIndex, VerticesToInsert );
+		}
+	}
+
+	// Generate normals and tangents
+	GenerateNormalsAndTangentsForPolygons( AffectedPolygons );
+
+	// Retriangulate all of the affected polygons
+	{
+		const bool bOnlyOnUndo = false;
+		RetriangulatePolygons( AffectedPolygons, bOnlyOnUndo );
+	}
+}
+
+
+void UEditableMesh::FindPolygonLoop( const FEdgeID EdgeID, TArray<FEdgeID>& OutEdgeLoopEdgeIDs, TArray<FEdgeID>& OutFlippedEdgeIDs, TArray<FEdgeID>& OutReversedEdgeIDPathToTake, TArray<FPolygonRef>& OutPolygonRefsToSplit ) const
+{
+	OutEdgeLoopEdgeIDs.Reset();
+	OutFlippedEdgeIDs.Reset();
+	OutReversedEdgeIDPathToTake.Reset();
+	OutPolygonRefsToSplit.Reset();
+
+	// Is the edge we're starting on a border edge?
+	bool bStartedOnBorderEdge = ( GetEdgeConnectedPolygonCount( EdgeID ) <= 1 );
+
+	// We'll actually do two passes searching for edges.  The first time, we'll flow along the polygons looking for a border edge.
+	// If we find one, it means we won't really have a loop, but instead we're just splitting a series of connected polygons.  In that
+	// case, the polygon with the border edge will become the start of our search for an opposing border edge.  Otherwise, we'll start
+	// at the input polygon and flow around until we come back to that polygon again.  If we can't make it back, then no polygons
+	// will be split by this operation.
+	bool bIsSearchingForBorderEdge = !bStartedOnBorderEdge;
+
+	// Keep track of whether we actually looped back around to the starting edge (rather than simply splitting across a series
+	// of polygons that both end at border edges.)
+	bool bIsCompleteLoop = false;
+
+	FEdgeID CurrentEdgeID = EdgeID;
+	bool bCurrentEdgeIsBorderEdge = bStartedOnBorderEdge;
+	bool bCurrentEdgeIsInOppositeDirection = false;
+	bool bCurrentEdgeIsInOppositeDirectionFromStartEdge = false;
+	for( ; ; )
+	{
+		// Add the current edge!
+		checkSlow( !OutEdgeLoopEdgeIDs.Contains( CurrentEdgeID ) );
+		OutEdgeLoopEdgeIDs.Add( CurrentEdgeID );
+		if( bCurrentEdgeIsInOppositeDirection )
+		{
+			OutFlippedEdgeIDs.Add( CurrentEdgeID );
+		}
+
+		FVertexID CurrentEdgeVertexIDs[ 2 ];
+		GetEdgeVertices( CurrentEdgeID, /* Out */ CurrentEdgeVertexIDs[ 0 ], /* Out */ CurrentEdgeVertexIDs[ 1 ] );
+
+		const FVector CurrentEdgeVertex0 = GetVertexAttribute( CurrentEdgeVertexIDs[ 0 ], UEditableMeshAttribute::VertexPosition(), 0 );
+		const FVector CurrentEdgeVertex1 = GetVertexAttribute( CurrentEdgeVertexIDs[ 1 ], UEditableMeshAttribute::VertexPosition(), 0 );
+		const FVector CurrentEdgeDirection = ( CurrentEdgeVertex1 - CurrentEdgeVertex0 ).GetSafeNormal();
+
+		const FEdgeID NextEdgeIDInPath = OutReversedEdgeIDPathToTake.Num() > 0 ? OutReversedEdgeIDPathToTake.Pop( false ) : FEdgeID::Invalid;
+
+		FEdgeID BestEdgeID = FEdgeID::Invalid;
+		FPolygonRef BestEdgeSplitsPolygon = FPolygonRef::Invalid;
+		bool bBestEdgeIsInOppositeDirection = false;
+		bool bBestEdgeIsBorderEdge = false;
+		float LargestAbsDotProduct = -1.0f;
+
+		// Let's take a look at all of the polygons connected to this edge.  These will start our loop.
+		const int32 ConnectedPolygonCount = GetEdgeConnectedPolygonCount( CurrentEdgeID );
+		for( int32 ConnectedPolygonNumber = 0; ConnectedPolygonNumber < ConnectedPolygonCount; ++ConnectedPolygonNumber )
+		{
+			const FPolygonRef ConnectedPolygonRef = GetEdgeConnectedPolygon( CurrentEdgeID, ConnectedPolygonNumber );
+
+			// Don't bother looking at the last polygon that was added to our split list.  We never want to backtrack!
+			if( OutPolygonRefsToSplit.Num() == 0 || ConnectedPolygonRef != OutPolygonRefsToSplit.Last() )
+			{
+				static TArray<FEdgeID> CandidateEdgeIDs;
+				CandidateEdgeIDs.Reset();
+				GetPolygonPerimeterEdges( ConnectedPolygonRef, /* Out */ CandidateEdgeIDs );
+
+				// Which edge of the connected polygon will be at the other end of our split?
+				for( const FEdgeID CandidateEdgeID : CandidateEdgeIDs )
+				{
+					// Don't bother with the edge we just came from
+					if( CandidateEdgeID != CurrentEdgeID )
+					{
+						// If we need to follow a specific path, then do that
+						if( NextEdgeIDInPath == FEdgeID::Invalid || CandidateEdgeID == NextEdgeIDInPath )
+						{
+							FVertexID CandidateEdgeVertexIDs[ 2 ];
+							GetEdgeVertices( CandidateEdgeID, /* Out */ CandidateEdgeVertexIDs[ 0 ], /* Out */ CandidateEdgeVertexIDs[ 1 ] );
+
+							const bool bIsBorderEdge = ( GetEdgeConnectedPolygonCount( CandidateEdgeID ) == 1 );
+
+							const FVector CandidateEdgeVertex0 = GetVertexAttribute( CandidateEdgeVertexIDs[ 0 ], UEditableMeshAttribute::VertexPosition(), 0 );
+							const FVector CandidateEdgeVertex1 = GetVertexAttribute( CandidateEdgeVertexIDs[ 1 ], UEditableMeshAttribute::VertexPosition(), 0 );
+							const FVector CandidateEdgeDirection = ( CandidateEdgeVertex1 - CandidateEdgeVertex0 ).GetSafeNormal();
+
+							const float DotProduct = FVector::DotProduct( CurrentEdgeDirection, CandidateEdgeDirection );
+							const float AbsDotProduct = FMath::Abs( DotProduct );
+
+							const float SameEdgeDirectionDotEpsilon = 0.05f;	// @todo mesheditor tweak
+							if( FMath::IsNearlyEqual( AbsDotProduct, LargestAbsDotProduct, SameEdgeDirectionDotEpsilon ) )
+							{
+								// If the candidate edge directions are pretty much the same, we'll choose the edge that flows closest to the
+								// direction that we split the last polygon in
+								if( OutEdgeLoopEdgeIDs.Num() > 1 )
+								{
+									FVertexID LastSplitEdgeVertexIDs[ 2 ];
+									GetEdgeVertices( OutEdgeLoopEdgeIDs[ OutEdgeLoopEdgeIDs.Num() - 2 ], /* Out */ LastSplitEdgeVertexIDs[ 0 ], /* Out */ LastSplitEdgeVertexIDs[ 1 ] );
+
+									const FVector LastSplitEdgeVertex0 = GetVertexAttribute( LastSplitEdgeVertexIDs[ 0 ], UEditableMeshAttribute::VertexPosition(), 0 );
+									const FVector LastSplitEdgeVertex1 = GetVertexAttribute( LastSplitEdgeVertexIDs[ 1 ], UEditableMeshAttribute::VertexPosition(), 0 );
+
+									const FVector DirectionTowardCenterOfCurrentEdge =
+										( FMath::Lerp( CurrentEdgeVertex0, CurrentEdgeVertex1, 0.5f ) -
+											FMath::Lerp( LastSplitEdgeVertex0, LastSplitEdgeVertex1, 0.5f ) ).GetSafeNormal();
+
+									const FVector DirectionTowardCenterOfCandidateEdge =
+										( FMath::Lerp( CandidateEdgeVertex0, CandidateEdgeVertex1, 0.5f ) -
+											FMath::Lerp( CurrentEdgeVertex0, CurrentEdgeVertex1, 0.5f ) ).GetSafeNormal();
+									const float CandidateEdgeDot = FVector::DotProduct( DirectionTowardCenterOfCurrentEdge, DirectionTowardCenterOfCandidateEdge );
+
+									check( BestEdgeID != FEdgeID::Invalid );
+
+									FVertexID BestEdgeVertexIDs[ 2 ];
+									GetEdgeVertices( BestEdgeID, /* Out */ BestEdgeVertexIDs[ 0 ], /* Out */ BestEdgeVertexIDs[ 1 ] );
+
+									const FVector BestEdgeVertex0 = GetVertexAttribute( BestEdgeVertexIDs[ 0 ], UEditableMeshAttribute::VertexPosition(), 0 );
+									const FVector BestEdgeVertex1 = GetVertexAttribute( BestEdgeVertexIDs[ 1 ], UEditableMeshAttribute::VertexPosition(), 0 );
+
+									const FVector DirectionTowardCenterOfBestEdge =
+										( FMath::Lerp( BestEdgeVertex0, BestEdgeVertex1, 0.5f ) -
+											FMath::Lerp( CurrentEdgeVertex0, CurrentEdgeVertex1, 0.5f ) ).GetSafeNormal();
+
+									const float BestEdgeDot = FVector::DotProduct( DirectionTowardCenterOfCurrentEdge, DirectionTowardCenterOfBestEdge );
+
+									if( CandidateEdgeDot > BestEdgeDot )
+									{
+										BestEdgeID = CandidateEdgeID;
+										BestEdgeSplitsPolygon = ConnectedPolygonRef;
+										bBestEdgeIsInOppositeDirection = ( DotProduct < 0.0f );
+										bBestEdgeIsBorderEdge = bIsBorderEdge;
+										LargestAbsDotProduct = AbsDotProduct;
+									}
+								}
+								else
+								{
+									// Edge directions are the same, but this is the very first split so we don't have a "flow" direction yet.
+									// Go ahead and prefer the edge that is closer to the initial edge.  This helps in the (uncommon) case of 
+									// multiple colinear edges on the same polygon (such as after SplitEdge() is called to insert a vertex on a polygon.)
+									const float BestEdgeDistance = [this, BestEdgeID, CurrentEdgeVertex0, CurrentEdgeVertex1]() -> float
+									{
+										check( BestEdgeID != FEdgeID::Invalid );
+
+										FVertexID BestEdgeVertexIDs[ 2 ];
+										GetEdgeVertices( BestEdgeID, /* Out */ BestEdgeVertexIDs[ 0 ], /* Out */ BestEdgeVertexIDs[ 1 ] );
+
+										const FVector BestEdgeVertex0 = GetVertexAttribute( BestEdgeVertexIDs[ 0 ], UEditableMeshAttribute::VertexPosition(), 0 );
+										const FVector BestEdgeVertex1 = GetVertexAttribute( BestEdgeVertexIDs[ 1 ], UEditableMeshAttribute::VertexPosition(), 0 );
+
+										FVector ClosestPoint0, ClosestPoint1;
+										FMath::SegmentDistToSegmentSafe( CurrentEdgeVertex0, CurrentEdgeVertex1, BestEdgeVertex0, BestEdgeVertex1, /* Out */ ClosestPoint0, /* Out */ ClosestPoint1 );
+										return ( ClosestPoint1 - ClosestPoint0 ).Size();
+									}( );
+
+									const float CandidateEdgeDistance = [CurrentEdgeVertex0, CurrentEdgeVertex1, CandidateEdgeVertex0, CandidateEdgeVertex1]() -> float
+									{
+										FVector ClosestPoint0, ClosestPoint1;
+										FMath::SegmentDistToSegmentSafe( CurrentEdgeVertex0, CurrentEdgeVertex1, CandidateEdgeVertex0, CandidateEdgeVertex1, /* Out */ ClosestPoint0, /* Out */ ClosestPoint1 );
+										return ( ClosestPoint1 - ClosestPoint0 ).Size();
+									}( );
+
+									if( CandidateEdgeDistance < BestEdgeDistance )
+									{
+										BestEdgeID = CandidateEdgeID;
+										BestEdgeSplitsPolygon = ConnectedPolygonRef;
+										bBestEdgeIsInOppositeDirection = ( DotProduct < 0.0f );
+										bBestEdgeIsBorderEdge = bIsBorderEdge;
+										LargestAbsDotProduct = AbsDotProduct;
+									}
+								}
+							}
+							else if( AbsDotProduct > LargestAbsDotProduct )
+							{
+								// This edge angle is the closest to our current edge so far!
+								BestEdgeID = CandidateEdgeID;
+								BestEdgeSplitsPolygon = ConnectedPolygonRef;
+								bBestEdgeIsInOppositeDirection = ( DotProduct < 0.0f );
+								bBestEdgeIsBorderEdge = bIsBorderEdge;
+								LargestAbsDotProduct = AbsDotProduct;
+							}
+						}
+					}
+				}
+			}
+		}
+
+		if( BestEdgeID != FEdgeID::Invalid &&
+			!OutPolygonRefsToSplit.Contains( BestEdgeSplitsPolygon ) )	// If we try to re-split the same polygon twice, then this loop is not valid
+		{
+			// OK, this polygon will definitely be split
+			OutPolygonRefsToSplit.Add( BestEdgeSplitsPolygon );
+
+			CurrentEdgeID = BestEdgeID;
+			bCurrentEdgeIsBorderEdge = bBestEdgeIsBorderEdge;
+			bCurrentEdgeIsInOppositeDirection = bBestEdgeIsInOppositeDirection;
+			if( bBestEdgeIsInOppositeDirection )
+			{
+				bCurrentEdgeIsInOppositeDirectionFromStartEdge = !bCurrentEdgeIsInOppositeDirectionFromStartEdge;
+			}
+
+			// Is the best edge already part of our loop?  If so, then we're done!
+			if( OutEdgeLoopEdgeIDs[ 0 ] == BestEdgeID )
+			{
+				bIsCompleteLoop = true;
+				break;
+			}
+			else if( OutEdgeLoopEdgeIDs.Contains( BestEdgeID ) )	// @todo mesheditor perf: Might want to use a TSet here, but we'll still need to keep a list too (order is important.)
+			{
+				// We ended up back at an edge that we already split, but it wasn't the edge that we started on.  This is
+				// not a valid loop, so clear our path and bail out.
+
+				// @todo mesheditor edgeloop: We need to revisit how we handle border edges and non-manifold edges.  Basically instead of searching
+				// for border edges, we should try looping around the "back" of polygons we've already visited, trying to get back to the starting
+				// edge.  This will allow the case of a cube mesh with a single non-manifold face to be split properly.  We just need to make sure
+				// we're not trying to split the same edge multiple times in different locations.  Probably, we don't want to split the same polygon
+				// more than once either.
+				OutEdgeLoopEdgeIDs.Reset();
+				OutFlippedEdgeIDs.Reset();
+				OutPolygonRefsToSplit.Reset();
+
+				break;
+			}
+			else if( bBestEdgeIsBorderEdge && bIsSearchingForBorderEdge )
+			{
+				// We found a border edge, so stop the search.  We'll now start over at this edge to form our loop.
+				bStartedOnBorderEdge = true;
+
+				bIsSearchingForBorderEdge = false;
+				bCurrentEdgeIsInOppositeDirection = bCurrentEdgeIsInOppositeDirectionFromStartEdge;
+
+				// Follow the path we took to get here, in reverse order, to make sure we get back to the edge we 
+				// were asked to create a loop on
+				OutReversedEdgeIDPathToTake = OutEdgeLoopEdgeIDs;
+
+				bIsCompleteLoop = false;
+
+				OutEdgeLoopEdgeIDs.Reset();
+				OutFlippedEdgeIDs.Reset();
+				OutPolygonRefsToSplit.Reset();
+			}
+			else
+			{
+				// Proceed to the next edge and try to continue the loop.  If we're at a border edge, the loop will definitely end here.
+			}
+		}
+		else
+		{
+			if( bStartedOnBorderEdge && bCurrentEdgeIsBorderEdge )
+			{
+				// We started on a border edge, and we've found the border edge on the other side of the polygons we'll be splitting.
+				// This isn't actually a loop, but we'll still split the polygons.
+			}
+			else
+			{
+				// We couldn't even find another edge, so we're done.
+				OutEdgeLoopEdgeIDs.Reset();
+				OutFlippedEdgeIDs.Reset();
+				OutPolygonRefsToSplit.Reset();
+			}
+
+			break;
+		}
+	}
+
+	// We're always splitting the same number of polygons as we have edges in the loop (these can be zero), except
+	// in the border edge case, where we're always splitting one less polygon.
+	if( bStartedOnBorderEdge &&
+		!bIsCompleteLoop )	// It's possible to start and finish on the same border edge, which is a complete loop
+	{
+		// We're splitting a series of polygons between two border edges
+		check( ( OutEdgeLoopEdgeIDs.Num() == 0 && OutPolygonRefsToSplit.Num() == 0 ) || ( OutEdgeLoopEdgeIDs.Num() == ( OutPolygonRefsToSplit.Num() + 1 ) ) );
+	}
+	else
+	{
+		// We're splitting polygons spanning a full loop of edges.  The starting edge is the same as the beginning edge.
+		check( OutEdgeLoopEdgeIDs.Num() == OutPolygonRefsToSplit.Num() );
+	}
+}
+
+
+void UEditableMesh::InsertEdgeLoop( const FEdgeID EdgeID, const TArray<float>& Splits, TArray<FEdgeID>& OutNewEdgeIDs )
+{
+	OutNewEdgeIDs.Reset();
+
+	// NOTE: The incoming splits should always be between 0.0 and 1.0, representing progress along 
+	//       the edge from the edge's first vertex toward it's other vertex.  The order doesn't matter (we'll sort them.)
+
+	// @todo mesheditor: Test with concave polygons -- probably we need to disallow/avoid splits that cross a concave part
+
+
+	static TArray<FEdgeID> EdgeLoopEdgeIDs;
+	EdgeLoopEdgeIDs.Reset();
+
+	static TArray<FEdgeID> FlippedEdgeIDs;
+	FlippedEdgeIDs.Reset();
+
+	static TArray<FEdgeID> ReversedEdgeIDPathToTake;
+	ReversedEdgeIDPathToTake.Reset();
+
+	static TArray<FPolygonRef> PolygonRefsToSplit;
+	PolygonRefsToSplit.Reset();
+
+
+	FindPolygonLoop( 
+		EdgeID, 
+		/* Out */ EdgeLoopEdgeIDs,
+		/* Out */ FlippedEdgeIDs,
+		/* Out */ ReversedEdgeIDPathToTake,
+		/* Out */ PolygonRefsToSplit );
+
+	static TSet<FEdgeID> FlippedEdgeIDSet;
+	FlippedEdgeIDSet.Reset();
+	FlippedEdgeIDSet.Append( FlippedEdgeIDs );
+
+	check( Splits.Num() > 0 );
+
+	// Sort the split values smallest to largest.  We'll be adding a strip of vertices for each split, and the
+	// IDs for those new vertices need to be in order.
+	static TArray<float> SortedSplits;
+	SortedSplits.Reset();
+	SortedSplits = Splits;
+	if( SortedSplits.Num() > 1 )
+	{
+		SortedSplits.Sort();
+	}
+
+
+	if( PolygonRefsToSplit.Num() > 0 )
+	{
+		// Keep track of the new vertices create by splitting all of the edges.  For each edge we split, an array of
+		// vertex IDs for each split along that edge
+		static TArray<TArray<FVertexID>> NewVertexIDsForEachEdge;
+		NewVertexIDsForEachEdge.Reset();
+
+
+		// Now let's go through and create new vertices for the loops by splitting edges
+		{
+			for( int32 EdgeLoopEdgeNumber = 0; EdgeLoopEdgeNumber < EdgeLoopEdgeIDs.Num(); ++EdgeLoopEdgeNumber )
+			{
+				const FEdgeID EdgeLoopEdgeID = EdgeLoopEdgeIDs[ EdgeLoopEdgeNumber ];
+
+				// If the edge winds in the opposite direction from the last edge, we'll need to flip the split positions around
+				const bool bIsFlipped = FlippedEdgeIDSet.Contains( EdgeLoopEdgeID );
+				if( bIsFlipped )
+				{
+					static TArray<float> TempSplits;
+					TempSplits.SetNumUninitialized( SortedSplits.Num(), false );
+					for( int32 SplitIndex = 0; SplitIndex < SortedSplits.Num(); ++SplitIndex )
+					{
+						TempSplits[ SplitIndex ] = 1.0f - SortedSplits[ ( SortedSplits.Num() - 1 ) - SplitIndex ];
+					}
+					SortedSplits = TempSplits;
+				}
+
+				// Split this edge
+				// @todo mesheditor urgent edgeloop: We need to make sure we're not retriangulating and calculating normals over and over again for the same faces.
+				static TArray<FVertexID> CurrentVertexIDs;
+				CurrentVertexIDs.Reset();
+				SplitEdge( EdgeLoopEdgeID, SortedSplits, /* Out */ CurrentVertexIDs );
+
+				// If the edge winding is backwards, we'll reverse the order of the vertex IDs in our list
+				// @todo mesheditor urgent edgeloop: INCORRECT (bIsFlipped is supposed to swap from last time -- TempSplits.)  Also not sure if needed or not... might be double compensating here.
+				static TArray<FVertexID> NewVertexIDsForEdge;
+				NewVertexIDsForEdge.SetNum( CurrentVertexIDs.Num(), false );
+				for( int32 VertexNumber = 0; VertexNumber < CurrentVertexIDs.Num(); ++VertexNumber )
+				{
+					NewVertexIDsForEdge[ /*bIsFlipped ? */( ( CurrentVertexIDs.Num() - VertexNumber ) - 1 ) /*: VertexNumber */] = CurrentVertexIDs[ VertexNumber ];
+				}
+				NewVertexIDsForEachEdge.Add( NewVertexIDsForEdge );
+			}
+		}
+
+
+		// Time to create new polygons for the split faces (and delete the old ones)
+		{
+			static TArray<FPolygonToSplit> PolygonsToSplit;
+			PolygonsToSplit.Reset();
+
+			for( int32 PolygonToSplitIter = 0; PolygonToSplitIter < PolygonRefsToSplit.Num(); ++PolygonToSplitIter )
+			{
+				const FPolygonRef PolygonRef = PolygonRefsToSplit[ PolygonToSplitIter ];
+
+				FPolygonToSplit& PolygonToSplit = *new( PolygonsToSplit ) FPolygonToSplit();
+				PolygonToSplit.PolygonRef = PolygonRef;
+
+				// The first and second edges connected to this polygon that are being split up
+				const int32 FirstEdgeNumber = PolygonToSplitIter;
+				const int32 SecondEdgeNumber = ( PolygonToSplitIter + 1 ) % EdgeLoopEdgeIDs.Num();
+
+				const FEdgeID FirstSplitEdgeID = EdgeLoopEdgeIDs[ FirstEdgeNumber ];
+				const FEdgeID SecondSplitEdgeID = EdgeLoopEdgeIDs[ SecondEdgeNumber ];
+				check( FirstSplitEdgeID != SecondSplitEdgeID );
+
+				// The (ordered) list of new vertices that was created by splitting the first and second edge.  One for each split.
+				const TArray<FVertexID>& FirstSplitEdgeNewVertexIDs = NewVertexIDsForEachEdge[ FirstEdgeNumber ];
+				const TArray<FVertexID>& SecondSplitEdgeNewVertexIDs = NewVertexIDsForEachEdge[ SecondEdgeNumber ];
+
+				for( int32 SplitIter = 0; SplitIter < SortedSplits.Num(); ++SplitIter )
+				{
+					FVertexPair& NewVertexPair = *new( PolygonToSplit.VertexPairsToSplitAt ) FVertexPair();
+					NewVertexPair.VertexID0 = FirstSplitEdgeNewVertexIDs[ SplitIter ];
+					NewVertexPair.VertexID1 = SecondSplitEdgeNewVertexIDs[ SplitIter ];
+				}
+			}
+
+
+			// Actually split up the polygons
+			static TArray<FEdgeID> NewEdgeIDs;
+			NewEdgeIDs.Reset();
+			SplitPolygons( PolygonsToSplit, /* Out */ NewEdgeIDs );
+
+			OutNewEdgeIDs.Append( NewEdgeIDs );
+		}
+	}
+}
+
+
+void UEditableMesh::SplitPolygons( const TArray<FPolygonToSplit>& PolygonsToSplit, TArray<FEdgeID>& OutNewEdgeIDs )
+{
+	OutNewEdgeIDs.Reset();
+
+	static TArray<FPolygonToCreate> PolygonsToCreate;
+	PolygonsToCreate.Reset();
+
+	for( const FPolygonToSplit& PolygonToSplit : PolygonsToSplit )
+	{
+		const FPolygonRef PolygonRef = PolygonToSplit.PolygonRef;
+
+		// Get all of the polygon's vertices
+		static TArray<FVertexID> PerimeterVertexIDs;
+		PerimeterVertexIDs.Reset();
+		GetPolygonPerimeterVertices( PolygonRef, /* Out */ PerimeterVertexIDs );
+		
+
+		// Figure out where exactly we're splitting the polygon for these splits.  Remember, we support splitting the
+		// polygon multiple times at once.  The first and last split are the most interesting because we need to continue
+		// the original flow of the polygon after inserting our new edge.  For all of the new polygons in the middle, we'll 
+		// just create simple quads.
+		const int32 SplitCount = PolygonToSplit.VertexPairsToSplitAt.Num();
+
+		int32 LastPolygonVertexNumbers[ 2 ] = { INDEX_NONE, INDEX_NONE };
+		bool bLastPolygonWindsForward = false;
+
+		const int32 NumPolygonsToCreate = SplitCount + 1;
+		for( int32 PolygonIter = 0; PolygonIter < NumPolygonsToCreate; ++PolygonIter )
+		{
+			const FVertexPair& VertexPair = PolygonToSplit.VertexPairsToSplitAt[ FMath::Min( PolygonIter, NumPolygonsToCreate - 2 ) ];
+
+			const FVertexID FirstVertexID = VertexPair.VertexID0;
+			const FVertexID SecondVertexID = VertexPair.VertexID1;
+
+			const int32 FirstVertexNumber = PerimeterVertexIDs.IndexOfByKey( FirstVertexID );
+			check( FirstVertexNumber != INDEX_NONE );	// Incoming vertex ID must already be a part of this polygon!
+			const int32 SecondVertexNumber = PerimeterVertexIDs.IndexOfByKey( SecondVertexID );
+			check( SecondVertexNumber != INDEX_NONE );	// Incoming vertex ID must already be a part of this polygon!
+
+
+			FPolygonToCreate& NewPolygon = *new( PolygonsToCreate ) FPolygonToCreate();
+			NewPolygon.SectionID = PolygonRef.SectionID;
+
+			static TArray<int32> PerimeterVertexNumbers;
+			PerimeterVertexNumbers.Reset();
+
+			const bool bWindsForward = FirstVertexNumber < SecondVertexNumber;
+
+			int32 SmallerVertexNumber = bWindsForward ? FirstVertexNumber : SecondVertexNumber;
+			int32 LargerVertexNumber = bWindsForward ? SecondVertexNumber : FirstVertexNumber;
+
+			if( PolygonIter == 0 || PolygonIter == ( NumPolygonsToCreate - 1 ) )
+			{
+				// This is either the first or last new polygon 
+				const bool bIsFirstPolygon = ( PolygonIter == 0 );
+
+				// Add the vertices we created for the new edge that will split the polygon
+				if( bIsFirstPolygon )
+				{
+					PerimeterVertexNumbers.Add( SmallerVertexNumber );
+					PerimeterVertexNumbers.Add( LargerVertexNumber );
+				}
+				else
+				{
+					PerimeterVertexNumbers.Add( LargerVertexNumber );
+					PerimeterVertexNumbers.Add( SmallerVertexNumber );
+				}
+
+				// Now add all of the other vertices of the original polygon that are on this side of the split
+				if( bIsFirstPolygon )
+				{
+					for( int32 VertexNumber = ( LargerVertexNumber + 1 ) % PerimeterVertexIDs.Num();
+						 VertexNumber != SmallerVertexNumber;
+						 VertexNumber = ( VertexNumber + 1 ) % PerimeterVertexIDs.Num() )
+					{
+						PerimeterVertexNumbers.Add( VertexNumber );
+					}
+				}
+				else
+				{
+					for( int32 VertexNumber = ( SmallerVertexNumber + 1 ) % PerimeterVertexIDs.Num();
+						 VertexNumber != LargerVertexNumber;
+						 VertexNumber = ( VertexNumber + 1 ) % PerimeterVertexIDs.Num() )
+					{
+						PerimeterVertexNumbers.Add( VertexNumber );
+					}
+				}
+			}
+			else
+			{
+				// @todo mesheditor urgent edgeloop: Polygon winding is incorrect with multiple splits
+
+				// This is a new polygon in the middle of other polygons created by the splits
+				PerimeterVertexNumbers.Add( bWindsForward ? SmallerVertexNumber : LargerVertexNumber );
+				PerimeterVertexNumbers.Add( bWindsForward ? LargerVertexNumber : SmallerVertexNumber );
+				PerimeterVertexNumbers.Add( LastPolygonVertexNumbers[ 1 ] );
+				PerimeterVertexNumbers.Add( LastPolygonVertexNumbers[ 0 ] );
+			}
+
+			NewPolygon.PerimeterVertices.Reserve( PerimeterVertexNumbers.Num() );
+			for( const int32 VertexNumber : PerimeterVertexNumbers )
+			{
+				FVertexAndAttributes& NewVertexAndAttributes = *new( NewPolygon.PerimeterVertices ) FVertexAndAttributes();
+				NewVertexAndAttributes.VertexID = PerimeterVertexIDs[ VertexNumber ];
+
+				// Copy all of the polygon vertex attributes over
+				{
+					for( int32 TextureCoordinateIndex = 0; TextureCoordinateIndex < TextureCoordinateCount; ++TextureCoordinateIndex )
+					{
+						NewVertexAndAttributes.PolygonVertexAttributes.Attributes.Add( FMeshElementAttributeData(
+							UEditableMeshAttribute::VertexTextureCoordinate(),
+							TextureCoordinateIndex,
+							GetPolygonPerimeterVertexAttribute( PolygonRef, VertexNumber, UEditableMeshAttribute::VertexTextureCoordinate(), TextureCoordinateIndex ) ) );
+					}
+
+					NewVertexAndAttributes.PolygonVertexAttributes.Attributes.Add( FMeshElementAttributeData(
+						UEditableMeshAttribute::VertexColor(),
+						0,
+						GetPolygonPerimeterVertexAttribute( PolygonRef, VertexNumber, UEditableMeshAttribute::VertexColor(), 0 ) ) );
+				}
+			}
+
+			LastPolygonVertexNumbers[ 0 ] = PerimeterVertexNumbers[ 0 ];
+			LastPolygonVertexNumbers[ 1 ] = PerimeterVertexNumbers[ 1 ];
+			bLastPolygonWindsForward = bWindsForward;
+		}
+	}
+
+	// Create new polygons that are split appropriately and connect to the new vertices we've added
+	static TArray< FPolygonRef > NewPolygonRefs;
+	NewPolygonRefs.Reset();
+	static TArray< FEdgeID > NewEdgeIDs;
+	NewEdgeIDs.Reset();
+	CreatePolygons( PolygonsToCreate, /* Out */ NewPolygonRefs, /* Out */ NewEdgeIDs );
+
+	OutNewEdgeIDs.Append( NewEdgeIDs );
+
+
+	// Delete the old polygons
+	{
+		static TArray<FPolygonRef> PolygonRefsToDelete;
+		PolygonRefsToDelete.Reset();
+		PolygonRefsToDelete.Reserve( PolygonsToSplit.Num() );
+		for( const FPolygonToSplit& PolygonToSplit : PolygonsToSplit )
+		{
+			PolygonRefsToDelete.Add( PolygonToSplit.PolygonRef );
+		}
+
+		const bool bDeleteOrphanEdges = false;
+		const bool bDeleteOrphanVertices = false;
+		const bool bDeleteEmptySections = false;
+		DeletePolygons( PolygonRefsToDelete, bDeleteOrphanEdges, bDeleteOrphanVertices, bDeleteEmptySections );
+	}
+
+	// Generate normals and tangents
+	GenerateNormalsAndTangentsForPolygons( NewPolygonRefs );
+}
+
+
+void UEditableMesh::DeleteEdgeAndConnectedPolygons( const FEdgeID EdgeID, const bool bDeleteOrphanedEdges, const bool bDeleteOrphanedVertices, const bool bDeleteEmptySections )
+{
+	static TArray<FPolygonRef> PolygonRefsToDelete;
+	PolygonRefsToDelete.Reset();
+
+	const int32 ConnectedPolygonCount = GetEdgeConnectedPolygonCount( EdgeID );
+	for( int32 ConnectedPolygonNumber = 0; ConnectedPolygonNumber < ConnectedPolygonCount; ++ConnectedPolygonNumber )
+	{
+		const FPolygonRef PolygonRef = GetEdgeConnectedPolygon( EdgeID, ConnectedPolygonNumber );
+
+		// Although it can be uncommon, it's possible the edge is connecting the same polygon to itself.  We need to add uniquely.
+		PolygonRefsToDelete.AddUnique( PolygonRef );
+	}
+
+	// Delete the polygons
+	DeletePolygons( PolygonRefsToDelete, bDeleteOrphanedEdges, bDeleteOrphanedVertices, bDeleteEmptySections );
+
+	// If the caller asked us not deleted orphaned edges, our edge-to-delete will still be hanging around.  Let's go
+	// and delete it now.
+	if( !bDeleteOrphanedEdges )
+	{
+		// NOTE: Because we didn't delete any orphaned edges, the incoming edge ID should still be valid
+		static TArray<FEdgeID> EdgeIDsToDelete;
+		EdgeIDsToDelete.Reset();
+		EdgeIDsToDelete.Add( EdgeID );
+
+		// This edge MUST be an orphan!
+		check( GetEdgeConnectedPolygonCount( EdgeID ) == 0 );
+
+		DeleteEdges( EdgeIDsToDelete, bDeleteOrphanedVertices );
+	}
+}
+
+
+void UEditableMesh::DeleteVertexAndConnectedEdgesAndPolygons( const FVertexID VertexID, const bool bDeleteOrphanedEdges, const bool bDeleteOrphanedVertices, const bool bDeleteEmptySections )
+{
+	static TArray<FEdgeID> EdgeIDsToDelete;
+	EdgeIDsToDelete.Reset();
+
+	const int32 ConnectedEdgeCount = GetVertexConnectedEdgeCount( VertexID );
+	for( int32 ConnectedEdgeNumber = 0; ConnectedEdgeNumber < ConnectedEdgeCount; ++ConnectedEdgeNumber )
+	{
+		const FEdgeID ConnectedEdgeID = GetVertexConnectedEdge( VertexID, ConnectedEdgeNumber );
+		EdgeIDsToDelete.Add( ConnectedEdgeID );
+	}
+
+	for( const FEdgeID EdgeIDToDelete : EdgeIDsToDelete )
+	{
+		// Make sure the edge still exists.  It may have been deleted as a polygon's edges were deleted during
+		// a previous iteration through this loop.
+		if( IsValidEdge( EdgeIDToDelete ) )
+		{
+			DeleteEdgeAndConnectedPolygons( EdgeIDToDelete, bDeleteOrphanedEdges, bDeleteOrphanedVertices, bDeleteEmptySections );
+		}
+	}
+}
+
+
+void UEditableMesh::CreateEmptyVertexRange( const int32 NumVerticesToAdd, TArray<FVertexID>& OutNewVertexIDs )
+{
+	CreateEmptyVertexRange_Internal( NumVerticesToAdd, nullptr, OutNewVertexIDs );
+
+	// NOTE: We iterate backwards, to delete vertices in the opposite order that we added them
+	FDeleteOrphanVerticesChangeInput RevertInput;
+	for( int32 VertexIter = NumVerticesToAdd - 1; VertexIter >= 0; --VertexIter )
+	{
+		const FVertexID VertexIDToDelete = OutNewVertexIDs[ VertexIter ];
+		RevertInput.VertexIDsToDelete.Add( VertexIDToDelete );
+	}
+
+	AddUndo( MakeUnique<FDeleteOrphanVerticesChange>( MoveTemp( RevertInput ) ) );
+}
+
+
+void UEditableMesh::DeleteOrphanVertices( const TArray<FVertexID>& VertexIDsToDelete )
+{
+	FCreateVerticesChangeInput RevertInput;
+
+
+	// Back everything up
+	{
+		// NOTE: We iterate backwards, to restore vertices in the opposite order that we deleted them
+		for( int32 VertexNumber = VertexIDsToDelete.Num() - 1; VertexNumber >= 0; --VertexNumber )
+		{
+			const FVertexID VertexID = VertexIDsToDelete[ VertexNumber ];
+
+			// Make sure the vertex is truly an orphan.  We're not going to be able to restore it's polygon vertex attributes,
+			// because the polygons won't exist when we're restoring the change
+			check( GetVertexConnectedEdgeCount( VertexID ) == 0 );
+
+			FVertexToCreate& VertexToCreate = *new( RevertInput.VerticesToCreate ) FVertexToCreate();
+			VertexToCreate.OriginalVertexID = VertexID;
+
+			for( const FName AttributeName : UEditableMesh::GetValidVertexAttributes() )
+			{
+				const int32 MaxAttributeIndex = GetMaxAttributeIndex( AttributeName );
+				for( int32 AttributeIndex = 0; AttributeIndex < MaxAttributeIndex; ++AttributeIndex )
+				{
+					FMeshElementAttributeData& VertexAttribute = *new( VertexToCreate.VertexAttributes.Attributes ) FMeshElementAttributeData(
+						AttributeName,
+						AttributeIndex,
+						GetVertexAttribute( VertexID, AttributeName, AttributeIndex ) );
+				}
+			}
+		}
+	}
+
+	DeleteOrphanVertices_Internal( VertexIDsToDelete );
+
+	AddUndo( MakeUnique<FCreateVerticesChange>( MoveTemp( RevertInput ) ) );
+}
+
+
+void UEditableMesh::DeleteEdges( const TArray<FEdgeID>& EdgeIDsToDelete, const bool bDeleteOrphanedVertices )
+{
+	// Back everything up
+	FCreateEdgesChangeInput RevertInput;
+	{
+		// NOTE: We iterate backwards, to restore edges in the opposite order that we deleted them
+		for( int32 EdgeNumber = EdgeIDsToDelete.Num() - 1; EdgeNumber >= 0; --EdgeNumber )
+		{
+			const FEdgeID EdgeID = EdgeIDsToDelete[ EdgeNumber ];
+
+			FEdgeToCreate& EdgeToCreate = *new( RevertInput.EdgesToCreate ) FEdgeToCreate();
+			EdgeToCreate.OriginalEdgeID = EdgeID;
+			EdgeToCreate.VertexID0 = GetEdgeVertex( EdgeID, 0 );
+			EdgeToCreate.VertexID1 = GetEdgeVertex( EdgeID, 1 );
+
+			const int32 ConnectedPolygonCount = GetEdgeConnectedPolygonCount( EdgeID );
+			EdgeToCreate.ConnectedPolygons.Reserve( ConnectedPolygonCount );
+			for( int32 ConnectedPolygonNumber = 0; ConnectedPolygonNumber < ConnectedPolygonCount; ++ConnectedPolygonNumber )
+			{
+				EdgeToCreate.ConnectedPolygons.Add( GetEdgeConnectedPolygon( EdgeID, ConnectedPolygonNumber ) );
+			}
+			
+			for( const FName AttributeName : UEditableMesh::GetValidEdgeAttributes() )
+			{
+				const int32 MaxAttributeIndex = GetMaxAttributeIndex( AttributeName );
+				for( int32 AttributeIndex = 0; AttributeIndex < MaxAttributeIndex; ++AttributeIndex )
+				{
+					FMeshElementAttributeData& EdgeAttribute = *new( EdgeToCreate.EdgeAttributes.Attributes ) FMeshElementAttributeData(
+						AttributeName,
+						AttributeIndex,
+						GetEdgeAttribute( EdgeID, AttributeName, AttributeIndex ) );
+				}
+			}
+		}
+	}
+
+
+	AddUndo( MoveTemp( MakeUnique<FCreateEdgesChange>( RevertInput ) ) );
+
+	// Delete the edges
+	DeleteEdges_Internal( EdgeIDsToDelete, bDeleteOrphanedVertices );
+}
+
+
+void UEditableMesh::CreateVertices( const TArray<FVertexToCreate>& VerticesToCreate, TArray<FVertexID>& OutNewVertexIDs )
+{
+	OutNewVertexIDs.Reset();
+
+	for( const FVertexToCreate& VertexToCreate : VerticesToCreate )
+	{
+		static TArray<FVertexID> OverrideVertexIDs;
+		OverrideVertexIDs.Reset();
+
+		if( VertexToCreate.OriginalVertexID != FVertexID::Invalid )
+		{
+			// Restore the vertex's original vertex ID
+			OverrideVertexIDs.Add( VertexToCreate.OriginalVertexID );
+		}
+
+		// Create the vertex and set it's attributes
+		static TArray<FVertexID> NewVertexIDs;
+		NewVertexIDs.Reset();
+		CreateEmptyVertexRange_Internal( 1, OverrideVertexIDs.Num() > 0 ? &OverrideVertexIDs : nullptr, /* Out */ NewVertexIDs );
+
+		const FVertexID NewVertexID = NewVertexIDs[ 0 ];
+		OutNewVertexIDs.Add( NewVertexID );
+
+		for( const FMeshElementAttributeData& VertexAttribute : VertexToCreate.VertexAttributes.Attributes )
+		{
+			SetVertexAttribute_Internal( NewVertexID, VertexAttribute.AttributeName, VertexAttribute.AttributeIndex, VertexAttribute.AttributeValue );
+		}
+	}
+
+	// NOTE: We iterate backwards, to delete vertices in the opposite order that we added them
+	FDeleteOrphanVerticesChangeInput RevertInput;
+	RevertInput.VertexIDsToDelete.Reserve( VerticesToCreate.Num() );
+	for( int32 VertexNumber = OutNewVertexIDs.Num() - 1; VertexNumber >= 0; --VertexNumber )
+	{
+		RevertInput.VertexIDsToDelete.Add( OutNewVertexIDs[ VertexNumber ] );
+	}
+
+	AddUndo( MakeUnique<FDeleteOrphanVerticesChange>( MoveTemp( RevertInput ) ) );
+}
+
+
+void UEditableMesh::CreateEdges( const TArray<FEdgeToCreate>& EdgesToCreate, TArray<FEdgeID>& OutNewEdgeIDs )
+{
+	OutNewEdgeIDs.Reset();
+
+	for( const FEdgeToCreate& EdgeToCreate : EdgesToCreate )
+	{
+		// Create the edge and set it's attributes
+		FEdgeID NewEdgeID = FEdgeID::Invalid;
+		CreateEdge_Internal( EdgeToCreate.VertexID0, EdgeToCreate.VertexID1, EdgeToCreate.ConnectedPolygons, EdgeToCreate.OriginalEdgeID, /* Out */ NewEdgeID );
+
+		OutNewEdgeIDs.Add( NewEdgeID );
+
+		for( const FMeshElementAttributeData& EdgeAttribute : EdgeToCreate.EdgeAttributes.Attributes )
+		{
+			SetEdgeAttribute_Internal( NewEdgeID, EdgeAttribute.AttributeName, EdgeAttribute.AttributeIndex, EdgeAttribute.AttributeValue );
+		}
+	}
+
+	// NOTE: We iterate backwards, to delete edges in the opposite order that we added them
+	FDeleteEdgesChangeInput RevertInput;
+	RevertInput.bDeleteOrphanedVertices = false;	// Don't delete any vertices on revert.  We're only creating edges here, not vertices!
+	RevertInput.EdgeIDsToDelete.Reserve( EdgesToCreate.Num() );
+	for( int32 EdgeNumber = OutNewEdgeIDs.Num() - 1; EdgeNumber >= 0; --EdgeNumber )
+	{
+		RevertInput.EdgeIDsToDelete.Add( OutNewEdgeIDs[ EdgeNumber ] );
+	}
+
+	AddUndo( MakeUnique<FDeleteEdgesChange>( MoveTemp( RevertInput ) ) );
+}
+
+
+void UEditableMesh::CreatePolygons( const TArray<FPolygonToCreate>& PolygonsToCreate, TArray<FPolygonRef>& OutNewPolygonRefs, TArray<FEdgeID>& OutNewEdgeIDs )
+{
+	OutNewPolygonRefs.Reset();
+	OutNewEdgeIDs.Reset();
+
+	for( const FPolygonToCreate& PolygonToCreate : PolygonsToCreate )
+	{
+		static TArray<FVertexID> PerimeterVertexIDs;
+		{
+			PerimeterVertexIDs.Reset();
+			PerimeterVertexIDs.Reserve( PolygonToCreate.PerimeterVertices.Num() );
+			for( const FVertexAndAttributes& PerimeterVertex : PolygonToCreate.PerimeterVertices )
+			{
+				PerimeterVertexIDs.Add( PerimeterVertex.VertexID );
+			}
+		}
+
+		static TArray<TArray<FVertexID>> VertexIDsForEachHole;
+		{
+			VertexIDsForEachHole.Reset();
+			VertexIDsForEachHole.Reserve( PolygonToCreate.PolygonHoles.Num() );
+
+			for( const FPolygonHoleVertices& PolygonHoleVertices : PolygonToCreate.PolygonHoles )
+			{
+				TArray<FVertexID>& HoleVertexIDs = *new( VertexIDsForEachHole ) TArray<FVertexID>();
+				HoleVertexIDs.Reserve( PolygonHoleVertices.HoleVertices.Num() );
+				for( const FVertexAndAttributes& HoleVertex : PolygonHoleVertices.HoleVertices )
+				{
+					HoleVertexIDs.Add( HoleVertex.VertexID );
+				}
+			}
+		}
+
+		// Create the polygon (along with any missing edges that are needed to connect the new polygon's vertices)
+		static TArray<FEdgeID> NewEdgeIDsForThisPolygon;
+		FPolygonRef NewPolygonRef;
+		CreatePolygon_Internal( PolygonToCreate.SectionID, PerimeterVertexIDs, VertexIDsForEachHole, PolygonToCreate.OriginalPolygonID, /* Out */ NewPolygonRef, /* Out */ NewEdgeIDsForThisPolygon );
+
+		OutNewEdgeIDs.Append( NewEdgeIDsForThisPolygon );
+		OutNewPolygonRefs.Add( NewPolygonRef );
+
+		// Set the polygon's vertex attributes
+		{
+			// Polygon perimeter
+			if( PolygonToCreate.PerimeterVertices.Num() > 0 )
+			{
+				const int32 PerimeterVertexCount = PolygonToCreate.PerimeterVertices.Num();
+				for( int32 PerimeterVertexNumber = 0; PerimeterVertexNumber < PerimeterVertexCount; ++PerimeterVertexNumber )
+				{
+					const TArray<FMeshElementAttributeData>& PerimeterVertexAttributeList = PolygonToCreate.PerimeterVertices[ PerimeterVertexNumber ].PolygonVertexAttributes.Attributes;
+					for( const FMeshElementAttributeData& PerimeterVertexAttribute : PerimeterVertexAttributeList )
+					{
+						SetPolygonPerimeterVertexAttribute_Internal( NewPolygonRef, PerimeterVertexNumber, PerimeterVertexAttribute.AttributeName, PerimeterVertexAttribute.AttributeIndex, PerimeterVertexAttribute.AttributeValue );
+					}
+				}
+			}
+
+			// Polygon holes
+			if( PolygonToCreate.PolygonHoles.Num() > 0 )
+			{
+				const int32 HoleCount = PolygonToCreate.PolygonHoles.Num();
+				for( int32 HoleNumber = 0; HoleNumber < HoleCount; ++HoleNumber )
+				{
+					const FPolygonHoleVertices& PolygonHoleVertices = PolygonToCreate.PolygonHoles[ HoleNumber ];
+					if( PolygonHoleVertices.HoleVertices.Num() > 0 )
+					{
+						const int32 HoleVertexCount = PolygonHoleVertices.HoleVertices.Num();
+						for( int32 HoleVertexNumber = 0; HoleVertexNumber < HoleVertexCount; ++HoleVertexNumber )
+						{
+							const TArray<FMeshElementAttributeData>& HoleVertexAttributeList = PolygonHoleVertices.HoleVertices[ HoleVertexNumber ].PolygonVertexAttributes.Attributes;
+							for( const FMeshElementAttributeData& HoleVertexAttribute : HoleVertexAttributeList )
+							{
+								SetPolygonHoleVertexAttribute_Internal( NewPolygonRef, HoleNumber, HoleVertexNumber, HoleVertexAttribute.AttributeName, HoleVertexAttribute.AttributeIndex, HoleVertexAttribute.AttributeValue );
+							}
+						}
+					}
+				}
+			}
+		}
+	}
+
+	// NOTE: We iterate backwards, to delete polygons in the opposite order that we added them
+	FDeletePolygonsChangeInput RevertInput;
+	RevertInput.PolygonRefsToDelete.Reserve( PolygonsToCreate.Num() );
+	RevertInput.bDeleteOrphanedEdges = false;	// Don't delete any edges on revert.  We're only creating polygons here, not vertices!
+	RevertInput.bDeleteOrphanedVertices = false;	// Don't delete any vertices on revert.  We're only creating polygons here, not vertices!
+	RevertInput.bDeleteEmptySections = false;
+	for( int32 PolygonNumber = OutNewPolygonRefs.Num() - 1; PolygonNumber >= 0; --PolygonNumber )
+	{
+		RevertInput.PolygonRefsToDelete.Add( OutNewPolygonRefs[ PolygonNumber ] );
+	}
+
+	AddUndo( MakeUnique<FDeletePolygonsChange>( MoveTemp( RevertInput ) ) );
+}
+
+
+void UEditableMesh::DeletePolygons( const TArray<FPolygonRef>& PolygonRefsToDelete, const bool bDeleteOrphanedEdges, const bool bDeleteOrphanedVertices, const bool bDeleteEmptySections )
+{
+	// Back everything up
+	FCreatePolygonsChangeInput RevertInput;
+	{
+		RevertInput.PolygonsToCreate.Reserve( PolygonRefsToDelete.Num() );
+
+		// NOTE: We iterate backwards, to restore edges in the opposite order that we deleted them
+		for( int32 PolygonNumber = PolygonRefsToDelete.Num() - 1; PolygonNumber >= 0; --PolygonNumber )
+		{
+			const FPolygonRef& PolygonRef = PolygonRefsToDelete[ PolygonNumber ];
+
+			FPolygonToCreate& PolygonToCreate = *new( RevertInput.PolygonsToCreate ) FPolygonToCreate();
+			PolygonToCreate.SectionID = PolygonRef.SectionID;
+			PolygonToCreate.OriginalPolygonID = PolygonRef.PolygonID;
+
+			static TArray<FVertexID> PolygonPerimeterVertexIDs;
+			GetPolygonPerimeterVertices( PolygonRef, /* Out */ PolygonPerimeterVertexIDs );
+
+			PolygonToCreate.PerimeterVertices.Reserve( PolygonPerimeterVertexIDs.Num() );
+			for( int32 PerimeterVertexNumber = 0; PerimeterVertexNumber < PolygonPerimeterVertexIDs.Num(); ++PerimeterVertexNumber )
+			{
+				const FVertexID PerimeterVertexID = PolygonPerimeterVertexIDs[ PerimeterVertexNumber ];
+
+				FVertexAndAttributes& VertexAndAttributes = *new( PolygonToCreate.PerimeterVertices ) FVertexAndAttributes();
+				VertexAndAttributes.VertexID = PerimeterVertexID;
+
+				for( const FName AttributeName : UEditableMesh::GetValidPolygonVertexAttributes() )
+				{
+					const int32 MaxAttributeIndex = GetMaxAttributeIndex( AttributeName );
+					for( int32 AttributeIndex = 0; AttributeIndex < MaxAttributeIndex; ++AttributeIndex )
+					{
+						FMeshElementAttributeData& PolygonVertexAttribute = *new( VertexAndAttributes.PolygonVertexAttributes.Attributes ) FMeshElementAttributeData(
+							AttributeName,
+							AttributeIndex,
+							GetPolygonPerimeterVertexAttribute( PolygonRef, PerimeterVertexNumber, AttributeName, AttributeIndex ) );
+					}
+				}
+			}
+
+			const int32 HoleCount = GetPolygonHoleCount( PolygonRef );
+			PolygonToCreate.PolygonHoles.SetNum( HoleCount, false );
+			for( int32 HoleNumber = 0; HoleNumber < HoleCount; ++HoleNumber )
+			{
+				FPolygonHoleVertices& Hole = PolygonToCreate.PolygonHoles[ HoleNumber ];
+
+				static TArray<FVertexID> PolygonHoleVertexIDs;
+				GetPolygonHoleVertices( PolygonRef, HoleNumber, /* Out */ PolygonHoleVertexIDs );
+
+				Hole.HoleVertices.Reserve( PolygonHoleVertexIDs.Num() );
+				for( int32 HoleVertexNumber = 0; HoleVertexNumber < PolygonHoleVertexIDs.Num(); ++HoleVertexNumber )
+				{
+					const FVertexID HoleVertexID = PolygonHoleVertexIDs[ HoleVertexNumber ];
+
+					FVertexAndAttributes& VertexAndAttributes = *new( Hole.HoleVertices ) FVertexAndAttributes();
+					VertexAndAttributes.VertexID = HoleVertexID;
+
+					for( const FName AttributeName : UEditableMesh::GetValidPolygonVertexAttributes() )
+					{
+						const int32 MaxAttributeIndex = GetMaxAttributeIndex( AttributeName );
+						for( int32 AttributeIndex = 0; AttributeIndex < MaxAttributeIndex; ++AttributeIndex )
+						{
+							FMeshElementAttributeData& PolygonVertexAttribute = *new( VertexAndAttributes.PolygonVertexAttributes.Attributes ) FMeshElementAttributeData(
+								AttributeName,
+								AttributeIndex,
+								GetPolygonHoleVertexAttribute( PolygonRef, HoleNumber, HoleVertexNumber, AttributeName, AttributeIndex ) );
+						}
+					}
+				}
+			}
+		}
+	}
+
+	AddUndo( MakeUnique<FCreatePolygonsChange>( MoveTemp( RevertInput ) ) );
+
+	for( const FPolygonRef& PolygonRefToDelete : PolygonRefsToDelete )
+	{
+		DeletePolygon_Internal( PolygonRefToDelete, bDeleteOrphanedEdges, bDeleteOrphanedVertices, bDeleteEmptySections );
+	}
+}
+
+
+void UEditableMesh::SetVerticesAttributes( const TArray<FAttributesForVertex>& AttributesForVertices )
+{
+	FSetVerticesAttributesChangeInput RevertInput;
+
+	RevertInput.AttributesForVertices.Reserve( AttributesForVertices.Num() );
+	for( const FAttributesForVertex& AttributesForVertex : AttributesForVertices )
+	{
+		FAttributesForVertex& RevertVertex = *new( RevertInput.AttributesForVertices ) FAttributesForVertex();
+		RevertVertex.VertexID = AttributesForVertex.VertexID;
+
+		// Set the vertex attributes
+		{
+			RevertVertex.VertexAttributes.Attributes.Reserve( AttributesForVertex.VertexAttributes.Attributes.Num() );
+
+			for( const FMeshElementAttributeData& VertexAttribute : AttributesForVertex.VertexAttributes.Attributes )
+			{
+				// Back up the attribute
+				FMeshElementAttributeData& RevertAttribute = *new( RevertVertex.VertexAttributes.Attributes ) FMeshElementAttributeData(
+					VertexAttribute.AttributeName,
+					VertexAttribute.AttributeIndex,
+					GetVertexAttribute( AttributesForVertex.VertexID, VertexAttribute.AttributeName, VertexAttribute.AttributeIndex ) );
+
+				// Set the new attribute
+				SetVertexAttribute_Internal( AttributesForVertex.VertexID, VertexAttribute.AttributeName, VertexAttribute.AttributeIndex, VertexAttribute.AttributeValue );
+			}
+		}
+	}
+
+	AddUndo( MakeUnique<FSetVerticesAttributesChange>( MoveTemp( RevertInput ) ) );
+}
+
+
+void UEditableMesh::SetEdgesAttributes( const TArray<FAttributesForEdge>& AttributesForEdges )
+{
+	FSetEdgesAttributesChangeInput RevertInput;
+
+	RevertInput.AttributesForEdges.Reserve( AttributesForEdges.Num() );
+	for( const FAttributesForEdge& AttributesForEdge : AttributesForEdges )
+	{
+		FAttributesForEdge& RevertEdge = *new( RevertInput.AttributesForEdges ) FAttributesForEdge();
+		RevertEdge.EdgeID = AttributesForEdge.EdgeID;
+
+		// Set the edge attributes
+		{
+			RevertEdge.EdgeAttributes.Attributes.Reserve( AttributesForEdge.EdgeAttributes.Attributes.Num() );
+
+			for( const FMeshElementAttributeData& EdgeAttribute : AttributesForEdge.EdgeAttributes.Attributes )
+			{
+				// Back up the attribute
+				FMeshElementAttributeData& RevertAttribute = *new( RevertEdge.EdgeAttributes.Attributes ) FMeshElementAttributeData(
+					EdgeAttribute.AttributeName,
+					EdgeAttribute.AttributeIndex,
+					GetEdgeAttribute( AttributesForEdge.EdgeID, EdgeAttribute.AttributeName, EdgeAttribute.AttributeIndex ) );
+
+				// Set the new attribute
+				SetEdgeAttribute_Internal( AttributesForEdge.EdgeID, EdgeAttribute.AttributeName, EdgeAttribute.AttributeIndex, EdgeAttribute.AttributeValue );
+			}
+		}
+	}
+
+	AddUndo( MakeUnique<FSetEdgesAttributesChange>( MoveTemp( RevertInput ) ) );
+}
+
+
+void UEditableMesh::SetPolygonsVertexAttributes( const TArray<FVertexAttributesForPolygon>& VertexAttributesForPolygons )
+{
+	FSetPolygonsVertexAttributesChangeInput RevertInput;
+
+	RevertInput.VertexAttributesForPolygons.Reserve( VertexAttributesForPolygons.Num() );
+	for( const FVertexAttributesForPolygon& VertexAttributesForPolygon : VertexAttributesForPolygons )
+	{
+		FVertexAttributesForPolygon& RevertPolygonVertexAttributes = *new( RevertInput.VertexAttributesForPolygons ) FVertexAttributesForPolygon();
+		RevertPolygonVertexAttributes.PolygonRef = VertexAttributesForPolygon.PolygonRef;
+
+		// Set the polygon's vertex attributes
+		{
+			// Polygon perimeter
+			if( VertexAttributesForPolygon.PerimeterVertexAttributeLists.Num() > 0 )
+			{
+				const int32 PerimeterVertexCount = GetPolygonPerimeterVertexCount( VertexAttributesForPolygon.PolygonRef );
+				check( VertexAttributesForPolygon.PerimeterVertexAttributeLists.Num() == PerimeterVertexCount );
+
+				RevertPolygonVertexAttributes.PerimeterVertexAttributeLists.SetNum( VertexAttributesForPolygon.PerimeterVertexAttributeLists.Num(), false );
+
+				for( int32 PerimeterVertexNumber = 0; PerimeterVertexNumber < PerimeterVertexCount; ++PerimeterVertexNumber )
+				{
+					const TArray<FMeshElementAttributeData>& PerimeterVertexAttributeList = VertexAttributesForPolygon.PerimeterVertexAttributeLists[ PerimeterVertexNumber ].Attributes;
+
+					RevertPolygonVertexAttributes.PerimeterVertexAttributeLists[ PerimeterVertexNumber ].Attributes.Reserve( PerimeterVertexAttributeList.Num() );
+					for( const FMeshElementAttributeData& PerimeterVertexAttribute : PerimeterVertexAttributeList )
+					{
+						// Back up the attribute
+						FMeshElementAttributeData& RevertAttribute = *new( RevertPolygonVertexAttributes.PerimeterVertexAttributeLists[ PerimeterVertexNumber ].Attributes ) FMeshElementAttributeData(
+							PerimeterVertexAttribute.AttributeName,
+							PerimeterVertexAttribute.AttributeIndex,
+							GetPolygonPerimeterVertexAttribute( VertexAttributesForPolygon.PolygonRef, PerimeterVertexNumber, PerimeterVertexAttribute.AttributeName, PerimeterVertexAttribute.AttributeIndex ) );
+
+						// Set the new attribute
+						SetPolygonPerimeterVertexAttribute_Internal( VertexAttributesForPolygon.PolygonRef, PerimeterVertexNumber, PerimeterVertexAttribute.AttributeName, PerimeterVertexAttribute.AttributeIndex, PerimeterVertexAttribute.AttributeValue );
+					}
+				}
+			}
+
+			// Polygon holes
+			if( VertexAttributesForPolygon.VertexAttributeListsForEachHole.Num() > 0 )
+			{
+				const int32 HoleCount = GetPolygonHoleCount( VertexAttributesForPolygon.PolygonRef );
+				check( VertexAttributesForPolygon.VertexAttributeListsForEachHole.Num() == HoleCount );
+
+				RevertPolygonVertexAttributes.VertexAttributeListsForEachHole.SetNum( HoleCount, false );
+
+				for( int32 HoleNumber = 0; HoleNumber < HoleCount; ++HoleNumber )
+				{
+					const FVertexAttributesForPolygonHole& AttributesForHoleVertices = VertexAttributesForPolygon.VertexAttributeListsForEachHole[ HoleNumber ];
+					if( AttributesForHoleVertices.VertexAttributeList.Num() > 0 )
+					{
+						const int32 HoleVertexCount = GetPolygonHoleVertexCount( VertexAttributesForPolygon.PolygonRef, HoleNumber );
+						check( AttributesForHoleVertices.VertexAttributeList.Num() == HoleVertexCount );
+
+						RevertPolygonVertexAttributes.VertexAttributeListsForEachHole[ HoleNumber ].VertexAttributeList.SetNum( HoleVertexCount, false );
+
+						for( int32 HoleVertexNumber = 0; HoleVertexNumber < HoleVertexCount; ++HoleVertexNumber )
+						{
+							const TArray<FMeshElementAttributeData>& HoleVertexAttributeList = AttributesForHoleVertices.VertexAttributeList[ HoleVertexNumber ].Attributes;
+
+							RevertPolygonVertexAttributes.VertexAttributeListsForEachHole[ HoleNumber ].VertexAttributeList[ HoleVertexNumber ].Attributes.Reserve( HoleVertexAttributeList.Num() );
+							for( const FMeshElementAttributeData& HoleVertexAttribute : HoleVertexAttributeList )
+							{
+								// Back up the attribute
+								FMeshElementAttributeData& RevertAttribute = *new( RevertPolygonVertexAttributes.VertexAttributeListsForEachHole[ HoleNumber ].VertexAttributeList[ HoleVertexNumber ].Attributes ) FMeshElementAttributeData(
+									HoleVertexAttribute.AttributeName,
+									HoleVertexAttribute.AttributeIndex,
+									GetPolygonHoleVertexAttribute( VertexAttributesForPolygon.PolygonRef, HoleNumber, HoleVertexNumber, HoleVertexAttribute.AttributeName, HoleVertexAttribute.AttributeIndex ) );
+
+								// Set the new attribute
+								SetPolygonHoleVertexAttribute_Internal( VertexAttributesForPolygon.PolygonRef, HoleNumber, HoleVertexNumber, HoleVertexAttribute.AttributeName, HoleVertexAttribute.AttributeIndex, HoleVertexAttribute.AttributeValue );
+							}
+						}
+					}
+				}
+			}
+		}
+	}
+
+	AddUndo( MakeUnique<FSetPolygonsVertexAttributesChange>( MoveTemp( RevertInput ) ) );
+}
+
+
+void UEditableMesh::TryToRemovePolygonEdge( const FEdgeID EdgeID, bool& bOutWasEdgeRemoved, FPolygonRef& OutNewPolygonRef )
+{
+	bOutWasEdgeRemoved = false;
+	OutNewPolygonRef = FPolygonRef::Invalid;
+
+	// If the edge is not shared by at least two polygons, we can't remove it.  (We would have to delete the polygon that owns
+	// this edge, which is not the intent of this feature.).  We also can't cleanly remove edges that are joining more than two 
+	// polygons.  We need to create a new polygon the two polygons, and if there were more than two then the remaining polygons 
+	// would be left disconnected after our edge is gone
+	const int32 ConnectedPolygonCount = GetEdgeConnectedPolygonCount( EdgeID );
+	if( ConnectedPolygonCount == 2 )
+	{
+		// Verify that both vertices on either end of this edge are connected to polygon (non-internal) edges.
+		// We currently do not expect to support internal triangles that don't touch the polygonal boundaries at all.
+		bool bBothVerticesConnectToPolygonEdges = true;
+		for( int32 EdgeVertexNumber = 0; EdgeVertexNumber < 2; ++EdgeVertexNumber )
+		{
+			const FVertexID VertexID = GetEdgeVertex( EdgeID, EdgeVertexNumber );
+
+			bool bVertexConnectsToPolygonEdge = false;
+
+			const int32 ConnectedEdgeCount = GetVertexConnectedEdgeCount( VertexID );
+			for( int32 ConnectedEdgeNumber = 0; ConnectedEdgeNumber < ConnectedEdgeCount; ++ConnectedEdgeNumber )
+			{
+				const FEdgeID OtherEdgeID = GetVertexConnectedEdge( VertexID, ConnectedEdgeNumber );
+				if( OtherEdgeID != EdgeID )
+				{
+					bVertexConnectsToPolygonEdge = true;
+					break;
+				}
+			}
+
+			if( !bVertexConnectsToPolygonEdge )
+			{
+				bBothVerticesConnectToPolygonEdges = false;
+			}
+		}
+
+		if( bBothVerticesConnectToPolygonEdges )
+		{
+			TArray<FVertexID> PolygonAVertexIDs, PolygonBVertexIDs;
+
+			const FPolygonRef PolygonARef = GetEdgeConnectedPolygon( EdgeID, 0 );
+			GetPolygonPerimeterVertices( PolygonARef, /* Out */ PolygonAVertexIDs );
+
+			const FPolygonRef PolygonBRef = GetEdgeConnectedPolygon( EdgeID, 1 );
+			GetPolygonPerimeterVertices( PolygonBRef, /* Out */ PolygonBVertexIDs );
+
+			// If the polygons are in different sections, we can't remove the edge because we can't determine
+			// which section the replacing polygon should belong to
+			// @todo mesheditor: We could just allow this, and use the first section ID.  It's just a bit weird in the UI.
+			if( PolygonARef.SectionID == PolygonBRef.SectionID )
+			{
+				const FSectionID NewPolygonSectionID = PolygonARef.SectionID;
+
+				// Create a polygon by combining the edges from either polygon we're connected to, omitting the edge we're removing
+				TArray<FVertexAndAttributes> NewPolygonVertices;
+				{
+					const FVertexID EdgeVertexIDA = GetEdgeVertex( EdgeID, 0 );
+					const FVertexID EdgeVertexIDB = GetEdgeVertex( EdgeID, 1 );
+
+					// @todo mesheditor urgent: This code was written with the assumption that the total vertex count would
+					// end up being two less than the total vertex count of the two polygons we're replacing.  However, this
+					// might not be true for cases where either polygon shares more than one edge with the other, right?
+
+					// Find the edge vertices in the first polygon
+					int32 EdgeStartsAtVertexInPolygonA = INDEX_NONE;
+					for( int32 PolygonAVertexNumber = 0; PolygonAVertexNumber < PolygonAVertexIDs.Num(); ++PolygonAVertexNumber )
+					{
+						const FVertexID PolygonAVertexID = PolygonAVertexIDs[ PolygonAVertexNumber ];
+						const FVertexID PolygonANextVertexID = PolygonAVertexIDs[ ( PolygonAVertexNumber + 1 ) % PolygonAVertexIDs.Num() ];
+
+						if( ( PolygonAVertexID == EdgeVertexIDA || PolygonAVertexID == EdgeVertexIDB ) &&
+							( PolygonANextVertexID == EdgeVertexIDA || PolygonANextVertexID == EdgeVertexIDB ) )
+						{
+							EdgeStartsAtVertexInPolygonA = PolygonAVertexNumber;
+							break;
+						}
+					}
+					check( EdgeStartsAtVertexInPolygonA != INDEX_NONE );
+					const int32 EdgeEndsAtVertexInPolygonA = ( EdgeStartsAtVertexInPolygonA + 1 ) % PolygonAVertexIDs.Num();
+
+
+					// Find the edge vertices in the second polygon
+					int32 EdgeStartsAtVertexInPolygonB = INDEX_NONE;
+					for( int32 PolygonBVertexNumber = 0; PolygonBVertexNumber < PolygonBVertexIDs.Num(); ++PolygonBVertexNumber )
+					{
+						const FVertexID PolygonBVertexID = PolygonBVertexIDs[ PolygonBVertexNumber ];
+						const FVertexID PolygonBNextVertexID = PolygonBVertexIDs[ ( PolygonBVertexNumber + 1 ) % PolygonBVertexIDs.Num() ];
+
+						if( ( PolygonBVertexID == EdgeVertexIDA || PolygonBVertexID == EdgeVertexIDB ) &&
+							( PolygonBNextVertexID == EdgeVertexIDA || PolygonBNextVertexID == EdgeVertexIDB ) )
+						{
+							EdgeStartsAtVertexInPolygonB = PolygonBVertexNumber;
+							break;
+						}
+					}
+					check( EdgeStartsAtVertexInPolygonB != INDEX_NONE );
+					const int32 EdgeEndsAtVertexInPolygonB = ( EdgeStartsAtVertexInPolygonB + 1 ) % PolygonBVertexIDs.Num();
+
+					
+					// Do the polygons wind in the same direction?  If they do, the edge order will be reversed.
+					const bool bPolygonsWindInSameDirection = PolygonAVertexIDs[ EdgeStartsAtVertexInPolygonA ] != PolygonBVertexIDs[ EdgeStartsAtVertexInPolygonB ];
+
+
+					// Start adding vertices from the first polygon, starting with the vertex right after the edge we're removing.
+					// We'll continue to add vertices from this polygon until we reach back around to that edge.
+					const int32 PolygonAStartVertex = EdgeEndsAtVertexInPolygonA;
+					const int32 PolygonAEndVertex = EdgeStartsAtVertexInPolygonA;
+					for( int32 PolygonAVertexNumber = PolygonAStartVertex; PolygonAVertexNumber != PolygonAEndVertex; )
+					{
+						const FVertexID PolygonAVertexID = PolygonAVertexIDs[ PolygonAVertexNumber ];
+
+						FVertexAndAttributes& NewPolygonVertex = *new( NewPolygonVertices ) FVertexAndAttributes();
+						NewPolygonVertex.VertexID = PolygonAVertexID;
+
+						// Store off vertex attributes so we can re-apply them later to the new polygon
+						// @todo mesheditor: utility method which automatically does this with all known polygon vertex attributes? (or redundant when using "uber mesh"?)
+						for( int32 TextureCoordinateIndex = 0; TextureCoordinateIndex < TextureCoordinateCount; ++TextureCoordinateIndex )
+						{
+							NewPolygonVertex.PolygonVertexAttributes.Attributes.Emplace(
+								UEditableMeshAttribute::VertexTextureCoordinate(),
+								TextureCoordinateIndex,
+								GetPolygonPerimeterVertexAttribute( PolygonARef, PolygonAVertexNumber, UEditableMeshAttribute::VertexTextureCoordinate(), TextureCoordinateIndex )
+							);
+						}
+
+						NewPolygonVertex.PolygonVertexAttributes.Attributes.Emplace(
+							UEditableMeshAttribute::VertexColor(),
+							0,
+							GetPolygonPerimeterVertexAttribute( PolygonARef, PolygonAVertexNumber, UEditableMeshAttribute::VertexColor(), 0 )
+						);
+
+						PolygonAVertexNumber = ( PolygonAVertexNumber + 1 ) % PolygonAVertexIDs.Num();
+					}
+
+					// Now add vertices from the second polygon
+					const int32 PolygonBStartVertex = bPolygonsWindInSameDirection ? EdgeEndsAtVertexInPolygonB : EdgeStartsAtVertexInPolygonB;
+					const int32 PolygonBEndVertex = bPolygonsWindInSameDirection ? EdgeStartsAtVertexInPolygonB : EdgeEndsAtVertexInPolygonB;
+					for( int32 PolygonBVertexNumber = PolygonBStartVertex; PolygonBVertexNumber != PolygonBEndVertex; )
+					{
+						const FVertexID PolygonBVertexID = PolygonBVertexIDs[ PolygonBVertexNumber ];
+
+						FVertexAndAttributes& NewPolygonVertex = *new( NewPolygonVertices ) FVertexAndAttributes();
+						NewPolygonVertex.VertexID = PolygonBVertexID;
+
+						// Store off vertex attributes so we can re-apply them later to the new polygon
+						for( int32 TextureCoordinateIndex = 0; TextureCoordinateIndex < TextureCoordinateCount; ++TextureCoordinateIndex )
+						{
+							NewPolygonVertex.PolygonVertexAttributes.Attributes.Emplace(
+								UEditableMeshAttribute::VertexTextureCoordinate(),
+								TextureCoordinateIndex,
+								GetPolygonPerimeterVertexAttribute( PolygonBRef, PolygonBVertexNumber, UEditableMeshAttribute::VertexTextureCoordinate(), TextureCoordinateIndex )
+							);
+						}
+
+						NewPolygonVertex.PolygonVertexAttributes.Attributes.Emplace(
+							UEditableMeshAttribute::VertexColor(),
+							0,
+							GetPolygonPerimeterVertexAttribute( PolygonBRef, PolygonBVertexNumber, UEditableMeshAttribute::VertexColor(), 0 )
+						);
+
+						if( bPolygonsWindInSameDirection )
+						{
+							// Iterate forwards
+							PolygonBVertexNumber = ( PolygonBVertexNumber + 1 ) % PolygonBVertexIDs.Num();
+						}
+						else
+						{
+							// Iterate backwards
+							if( PolygonBVertexNumber == 0 )	   // @todo mesheditor: Not tested yet.  Need to connect two faces that wind in opposite directions
+							{
+								PolygonBVertexNumber = PolygonBVertexIDs.Num() - 1;
+							}
+							else
+							{
+								--PolygonBVertexNumber;
+							}
+						}
+					}
+				}
+
+				// OK, we can go ahead and delete the edge and its connected polygons.  We do NOT want to delete any orphaned
+				// edges or vertices though.  We're going to create a new polygon that connects to those right afterwards.
+				const bool bDeleteOrphanedEdges = false;
+				const bool bDeleteOrphanedVertices = false;
+				const bool bDeleteEmptySections = false;
+				DeleteEdgeAndConnectedPolygons( EdgeID, bDeleteOrphanedEdges, bDeleteOrphanedVertices, bDeleteEmptySections );
+
+				// Now create a new polygon to replace the two polygons we deleted
+				{
+					// Create the polygon
+					static TArray<FPolygonToCreate> PolygonsToCreate;
+					PolygonsToCreate.Reset();
+
+					// @todo mesheditor perf: Ideally we support creating multiple polygons at once and batching up the work
+					FPolygonToCreate& PolygonToCreate = *new( PolygonsToCreate ) FPolygonToCreate();
+					PolygonToCreate.SectionID = NewPolygonSectionID;
+					PolygonToCreate.PerimeterVertices = NewPolygonVertices;	// @todo mesheditor perf: Extra allocatons/copies: Ideally MoveTemp() here but we can't move a STATIC local!
+					// @todo mesheditor hole: Hole support needed!
+
+					static TArray<FPolygonRef> NewPolygonRefs;
+					NewPolygonRefs.Reset();
+					static TArray<FEdgeID> NewEdgeIDs;
+					NewEdgeIDs.Reset();
+					CreatePolygons( PolygonsToCreate, /* Out */ NewPolygonRefs, /* Out */ NewEdgeIDs );	// @todo mesheditor perf: Extra allocatons/copies: Ideally MoveTemp() here but we can't move a STATIC local!
+
+					OutNewPolygonRef = NewPolygonRefs[ 0 ];
+				}
+
+				// The polygon's triangulation will be different, plus we lost an edge which may have been hard/soft, thus we need new normals. 
+				// Note that the hardness of the other edges will remain the same.
+				static TArray< FPolygonRef > NewPolygons;
+				NewPolygons.Reset();
+				NewPolygons.Add( OutNewPolygonRef );
+				GenerateNormalsAndTangentsForPolygonsAndAdjacents( NewPolygons );
+
+				bOutWasEdgeRemoved = true;
+			}
+		}
+	}
+}
+
+
+void UEditableMesh::TryToRemoveVertex( const FVertexID VertexID, bool& bOutWasVertexRemoved, FEdgeID& OutNewEdgeID )
+{
+	bOutWasVertexRemoved = false;
+	OutNewEdgeID = FEdgeID::Invalid;
+
+	// We only support removing vertices that are shared by just two edges
+	const int32 ConnectedEdgeCount = GetVertexConnectedEdgeCount( VertexID );
+	if( ConnectedEdgeCount == 2 )
+	{
+		// Get the two vertices on the other edge of either edge
+		FVertexID NewEdgeVertexIDs[ 2 ];
+		for( int32 EdgeNumber = 0; EdgeNumber < ConnectedEdgeCount; ++EdgeNumber )
+		{
+			const FEdgeID OtherEdgeID = GetVertexConnectedEdge( VertexID, EdgeNumber );
+			FVertexID OtherEdgeVertexIDs[ 2 ];
+			GetEdgeVertices( OtherEdgeID, /* Out */ OtherEdgeVertexIDs[0], /* Out */ OtherEdgeVertexIDs[1] );
+
+			NewEdgeVertexIDs[ EdgeNumber ] = OtherEdgeVertexIDs[ 0 ] == VertexID ? OtherEdgeVertexIDs[ 1 ] : OtherEdgeVertexIDs[ 0 ];
+		}
+
+		// Try to preserve attributes of the edges we're deleting.  We'll take the attributes from the first
+		// edge and apply them to the newly created edge
+		static TArray<FMeshElementAttributeData> NewEdgeAttributes;
+		{
+			const FEdgeID OtherEdgeID = GetVertexConnectedEdge( VertexID, 0 );		// @todo mesheditor: This is sort of subjective how we handle this (taking attributes from the first edge)
+
+			NewEdgeAttributes.Reset();
+			for( const FName AttributeName : UEditableMesh::GetValidEdgeAttributes() )
+			{
+				const int32 MaxAttributeIndex = GetMaxAttributeIndex( AttributeName );
+				for( int32 AttributeIndex = 0; AttributeIndex < MaxAttributeIndex; ++AttributeIndex )
+				{
+					FMeshElementAttributeData& EdgeAttribute = *new( NewEdgeAttributes ) FMeshElementAttributeData(
+						AttributeName,
+						AttributeIndex,
+						GetEdgeAttribute( OtherEdgeID, AttributeName, AttributeIndex ) );
+				}
+			}
+		}
+
+		// The new edge will be connected to the same polygons as both of the edges we're replacing.  Because
+		// we only support deleting a vertex shared by two edges, the two edges are guaranteed to be connected
+		// to the same exact polygons
+		static TArray<FPolygonRef> NewEdgeConnectedPolygons;
+		GetVertexConnectedPolygons( VertexID, /* Out */ NewEdgeConnectedPolygons );
+
+		// When undoing this change, we'll retriangulate after everything has been hooked up (before we start destroying edges and vertices below)
+		{
+			const bool bOnlyOnUndo = true;
+			RetriangulatePolygons( NewEdgeConnectedPolygons, bOnlyOnUndo );
+		}
+
+		// Remove the vertex from it's connected polygons
+		{
+			for( const FPolygonRef PolygonRef : NewEdgeConnectedPolygons )
+			{
+				// @todo mesheditor holes
+				const int32 PolygonVertexNumber = FindPolygonPerimeterVertexNumberForVertex( PolygonRef, VertexID );
+				check( PolygonVertexNumber != INDEX_NONE );
+				RemovePolygonPerimeterVertices( PolygonRef, PolygonVertexNumber, 1 );
+			}
+		}
+
+		// Delete the two edges
+		{
+			static TArray<FEdgeID> EdgeIDsToDelete;
+			EdgeIDsToDelete.Reset();
+			for( int32 EdgeNumber = 0; EdgeNumber < ConnectedEdgeCount; ++EdgeNumber )
+			{
+				EdgeIDsToDelete.Add( GetVertexConnectedEdge( VertexID, EdgeNumber ) );
+			}
+
+			// NOTE: We can't delete the orphan vertex yet because the polygon triangles are still referencing
+			// its rendering vertices.  We'll delete the edges, retriangulate, then delete the vertex afterwards.
+			const bool bDeleteOrphanedVertices = false;
+			DeleteEdges( EdgeIDsToDelete, bDeleteOrphanedVertices );
+		}
+
+		FEdgeID NewEdgeID = FEdgeID::Invalid;
+
+		// Create a new edge to replace the vertex and two edges we deleted
+		{
+			static TArray<FEdgeToCreate> EdgesToCreate;
+			EdgesToCreate.Reset();
+
+			FEdgeToCreate& EdgeToCreate = *new( EdgesToCreate ) FEdgeToCreate;
+			EdgeToCreate.VertexID0 = NewEdgeVertexIDs[ 0 ];
+			EdgeToCreate.VertexID1 = NewEdgeVertexIDs[ 1 ];
+			EdgeToCreate.ConnectedPolygons = NewEdgeConnectedPolygons;
+			EdgeToCreate.EdgeAttributes.Attributes = NewEdgeAttributes;
+
+			static TArray<FEdgeID> NewEdgeIDs;
+			NewEdgeIDs.Reset();
+			CreateEdges( EdgesToCreate, /* Out */ NewEdgeIDs );
+
+			NewEdgeID = NewEdgeIDs[ 0 ];
+		}
+
+		// Update the normals of the affected polygons
+		GenerateNormalsAndTangentsForPolygonsAndAdjacents( NewEdgeConnectedPolygons );
+
+		// Retriangulate all of the affected polygons
+		{
+			const bool bOnlyOnUndo = false;
+			RetriangulatePolygons( NewEdgeConnectedPolygons, bOnlyOnUndo );
+		}
+
+		// Delete the (orphaned) vertex
+		{
+			static TArray<FVertexID> OrphanVertexIDsToDelete;
+			OrphanVertexIDsToDelete.Reset();
+			OrphanVertexIDsToDelete.Add( VertexID );
+			DeleteOrphanVertices( OrphanVertexIDsToDelete );
+		}
+
+		bOutWasVertexRemoved = true;
+		OutNewEdgeID = NewEdgeID;
+	}
+}
+
+
+void UEditableMesh::ExtrudePolygons( const TArray<FPolygonRef>& Polygons, const float ExtrudeDistance, TArray<FPolygonRef>& OutNewExtrudedFrontPolygons )
+{
+	// @todo mesheditor: For extrusion, we really need a mode where selected polygons that are connected will all
+	// be extruded together (no new edges created between the neighbors.)  See other modeling programs for examples.
+	// --> This should also work with multiple selected edges and vertices too, ideally
+
+	// @todo mesheditor perf: We can make this much faster by batching up polygons together.  Just be careful about how neighbors are handled.
+	OutNewExtrudedFrontPolygons.Reset();
+	
+	static TArray<FPolygonRef> AllNewPolygons;
+	AllNewPolygons.Reset();
+
+	static TArray<FAttributesForEdge> AttributesForEdges;
+	AttributesForEdges.Reset();
+
+	static TArray<FAttributesForVertex> AttributesForVertices;
+	AttributesForVertices.Reset();
+
+	static TArray<FVertexAttributesForPolygon> VertexAttributesForPolygons;
+	VertexAttributesForPolygons.Reset();
+
+	for( int32 PolygonIter = 0; PolygonIter < Polygons.Num(); ++PolygonIter )
+	{
+		const FPolygonRef& PolygonRef = Polygons[ PolygonIter ];
+
+		// We'll need this polygon's normal to figure out where to put the extruded copy of the polygon
+		const FVector PolygonNormal = ComputePolygonNormal( PolygonRef );
+
+		// Create a new set of vertices for the extruded polygon
+		static TArray<FVertexID> ExtrudedVertexIDs;
+		ExtrudedVertexIDs.Reset();
+		{
+			const int32 PerimeterVertexCount = this->GetPolygonPerimeterVertexCount( PolygonRef );
+			ExtrudedVertexIDs.SetNumUninitialized( PerimeterVertexCount, false );
+
+			// Create some vertices
+			CreateEmptyVertexRange( PerimeterVertexCount, /* Out */ ExtrudedVertexIDs );
+		}
+		int32 NextAvailableExtrudedVertexIDNumber = 0;
+
+		static TArray<FVertexID> PolygonVertexIDs;
+		PolygonVertexIDs.Reset();
+		this->GetPolygonPerimeterVertices( PolygonRef, /* Out */ PolygonVertexIDs );
+
+		// Map all of the edge vertices to their new extruded counterpart
+		static TMap< FVertexID, FVertexID> VertexIDToExtrudedCopyMap;
+		VertexIDToExtrudedCopyMap.Reset();
+		{
+			const int32 PerimeterEdgeCount = this->GetPolygonPerimeterEdgeCount( PolygonRef );
+			for( int32 PerimeterEdgeNumber = 0; PerimeterEdgeNumber < PerimeterEdgeCount; ++PerimeterEdgeNumber )
+			{
+				FVertexID EdgeVertexIDs[ 2 ];
+				{
+					// @todo mesheditor perf: O(log N^2) iteration here. For every edge, for every edge up to this index.  Need to clean this up. 
+					//		--> Also, there are quite a few places where we are stepping through edges in perimeter-order.  We need to have a nice way to walk that.
+					bool bEdgeWindingIsReversedForPolygon;
+					const FEdgeID EdgeID = this->GetPolygonPerimeterEdge( PolygonRef, PerimeterEdgeNumber, /* Out */ bEdgeWindingIsReversedForPolygon );
+
+					// After extruding, all of the edges of the original polygon become hard edges
+					FAttributesForEdge& AttributesForEdge = *new( AttributesForEdges ) FAttributesForEdge();
+					AttributesForEdge.EdgeID = EdgeID;
+					AttributesForEdge.EdgeAttributes.Attributes.Add( FMeshElementAttributeData( UEditableMeshAttribute::EdgeIsHard(), 0, FVector4( 1.0f ) ) );
+					AttributesForEdge.EdgeAttributes.Attributes.Add( FMeshElementAttributeData( UEditableMeshAttribute::EdgeCreaseSharpness(), 0, FVector4( 1.0f ) ) );
+
+					this->GetEdgeVertices( EdgeID, /* Out */ EdgeVertexIDs[0], /* Out */ EdgeVertexIDs[1] );
+
+					if( bEdgeWindingIsReversedForPolygon )
+					{
+						const FVertexID TempVertexID = EdgeVertexIDs[ 1 ];
+						EdgeVertexIDs[ 1 ] = EdgeVertexIDs[ 0 ];
+						EdgeVertexIDs[ 0 ] = TempVertexID;
+					}
+				}
+
+				FVertexID ExtrudedEdgeVertexIDs[ 2 ];
+				for( int32 EdgeVertexNumber = 0; EdgeVertexNumber < 2; ++EdgeVertexNumber )
+				{
+					const FVertexID EdgeVertexID = EdgeVertexIDs[ EdgeVertexNumber ];
+					FVertexID* ExtrudedEdgeVertexIDPtr = VertexIDToExtrudedCopyMap.Find( EdgeVertexID );
+					if( ExtrudedEdgeVertexIDPtr != nullptr )
+					{
+						ExtrudedEdgeVertexIDs[ EdgeVertexNumber ] = *ExtrudedEdgeVertexIDPtr;
+					}
+					else
+					{
+						// Create a copy of this vertex for the extruded face
+						const FVertexID ExtrudedVertexID = ExtrudedVertexIDs[ NextAvailableExtrudedVertexIDNumber++ ];
+
+						ExtrudedEdgeVertexIDPtr = &VertexIDToExtrudedCopyMap.Add( EdgeVertexID, ExtrudedVertexID );
+
+						// Push the vertex out along the polygon's normal
+						const FVector OriginalVertexPosition = this->GetVertexAttribute( EdgeVertexID, UEditableMeshAttribute::VertexPosition(), 0 );
+
+						const FVector ExtrudedVertexPosition = OriginalVertexPosition + ExtrudeDistance * PolygonNormal;
+
+						// Fill in the vertex
+						FAttributesForVertex& AttributesForVertex = *new( AttributesForVertices ) FAttributesForVertex();
+						AttributesForVertex.VertexID = ExtrudedVertexID;
+						AttributesForVertex.VertexAttributes.Attributes.Add( FMeshElementAttributeData(
+							UEditableMeshAttribute::VertexPosition(), 
+							0, 
+							ExtrudedVertexPosition ) );
+					}
+					ExtrudedEdgeVertexIDs[ EdgeVertexNumber ] = *ExtrudedEdgeVertexIDPtr;
+				}
+
+				// Add a face that connects the edge to it's extruded counterpart edge
+				{
+					FVertexID OriginalVertexIDs[ 4 ];
+
+					static TArray<FVertexAndAttributes> NewSidePolygonVertices;
+					NewSidePolygonVertices.Reset( 4 );
+					NewSidePolygonVertices.SetNum( 4, false );	// Always four edges in an extruded face
+
+					NewSidePolygonVertices[ 0 ].VertexID = EdgeVertexIDs[ 1 ];
+					OriginalVertexIDs[ 0 ] = EdgeVertexIDs[ 1 ];
+					NewSidePolygonVertices[ 1 ].VertexID = EdgeVertexIDs[ 0 ];
+					OriginalVertexIDs[ 1 ] = EdgeVertexIDs[ 0 ];
+					NewSidePolygonVertices[ 2 ].VertexID = ExtrudedEdgeVertexIDs[ 0 ];
+					OriginalVertexIDs[ 2 ] = EdgeVertexIDs[ 0 ];
+					NewSidePolygonVertices[ 3 ].VertexID = ExtrudedEdgeVertexIDs[ 1 ];
+					OriginalVertexIDs[ 3 ] = EdgeVertexIDs[ 1 ];
+
+					FPolygonRef NewSidePolygonRef;	// Filled in below
+					{
+						static TArray<FPolygonToCreate> PolygonsToCreate;
+						PolygonsToCreate.Reset();
+
+						// Create the polygon
+						// @todo mesheditor perf: Ideally we support creating multiple polygons at once and batching up the work
+						FPolygonToCreate& PolygonToCreate = *new( PolygonsToCreate ) FPolygonToCreate();
+						PolygonToCreate.SectionID = PolygonRef.SectionID;
+						PolygonToCreate.PerimeterVertices = NewSidePolygonVertices;	// @todo mesheditor perf: Copying static array here, ideally allocations could be avoided
+
+						// NOTE: We never create holes in side polygons
+
+						static TArray<FPolygonRef> NewPolygonRefs;
+						NewPolygonRefs.Reset();
+						static TArray<FEdgeID> NewEdgeIDs;
+						NewEdgeIDs.Reset();
+						CreatePolygons( PolygonsToCreate, /* Out */ NewPolygonRefs, /* Out */ NewEdgeIDs );	// @todo mesheditor perf: Extra allocatons/copies: Ideally MoveTemp() here but we can't move a STATIC local!
+
+						NewSidePolygonRef = NewPolygonRefs[ 0 ];
+					}
+					AllNewPolygons.Add( NewSidePolygonRef );
+
+					// Copy polygon UVs from the original polygon's vertices
+					{
+						FVertexAttributesForPolygon& PolygonNewAttributes = *new( VertexAttributesForPolygons ) FVertexAttributesForPolygon();
+						PolygonNewAttributes.PolygonRef = NewSidePolygonRef;
+						PolygonNewAttributes.PerimeterVertexAttributeLists.SetNum( NewSidePolygonVertices.Num(), false );
+
+						for( int32 NewVertexIter = 0; NewVertexIter < NewSidePolygonVertices.Num(); ++NewVertexIter )
+						{
+							const int32 PolygonVertexNumber = this->FindPolygonPerimeterVertexNumberForVertex( PolygonRef, OriginalVertexIDs[ NewVertexIter ] );
+							check( PolygonVertexNumber != INDEX_NONE );
+
+							TArray<FMeshElementAttributeData>& PerimeterVertexNewAttributes = PolygonNewAttributes.PerimeterVertexAttributeLists[ NewVertexIter ].Attributes;
+							for( int32 TextureCoordinateIndex = 0; TextureCoordinateIndex < TextureCoordinateCount; ++TextureCoordinateIndex )
+							{
+								const FVector4 TextureCoordinate = this->GetPolygonPerimeterVertexAttribute(
+									PolygonRef,
+									PolygonVertexNumber,
+									UEditableMeshAttribute::VertexTextureCoordinate(),
+									TextureCoordinateIndex );
+
+								PerimeterVertexNewAttributes.Emplace( UEditableMeshAttribute::VertexTextureCoordinate(), TextureCoordinateIndex, TextureCoordinate );
+							}
+
+							const FVector4 VertexColor = this->GetPolygonPerimeterVertexAttribute(
+								PolygonRef,
+								PolygonVertexNumber,
+								UEditableMeshAttribute::VertexColor(),
+								0 );
+
+							PerimeterVertexNewAttributes.Emplace( UEditableMeshAttribute::VertexColor(), 0, VertexColor );
+						}
+					}
+
+					// All of this edges of the new polygon will be hard
+					{
+						const int32 NewPerimeterEdgeCount = this->GetPolygonPerimeterEdgeCount( NewSidePolygonRef );
+						for( int32 NewPerimeterEdgeNumber = 0; NewPerimeterEdgeNumber < NewPerimeterEdgeCount; ++NewPerimeterEdgeNumber )
+						{
+							bool bEdgeWindingIsReversedForPolygon;
+							const FEdgeID EdgeID = this->GetPolygonPerimeterEdge( NewSidePolygonRef, NewPerimeterEdgeNumber, /* Out */ bEdgeWindingIsReversedForPolygon );
+
+							// New side polygons get hard edges, but not hard crease weights
+							FAttributesForEdge& AttributesForEdge = *new( AttributesForEdges ) FAttributesForEdge();
+							AttributesForEdge.EdgeID = EdgeID;
+							AttributesForEdge.EdgeAttributes.Attributes.Add( FMeshElementAttributeData( UEditableMeshAttribute::EdgeIsHard(), 0, FVector4( 1.0f ) ) );
+						}
+					}
+				}
+			}
+		}
+		check( NextAvailableExtrudedVertexIDNumber == ExtrudedVertexIDs.Num() );
+
+		// Create a new extruded polygon for the face
+		FPolygonRef ExtrudedFrontPolygonRef;	// Filled in below
+		{
+			static TArray<FVertexAndAttributes> NewFrontPolygonVertices;
+			NewFrontPolygonVertices.Reset( PolygonVertexIDs.Num() );
+			NewFrontPolygonVertices.SetNum( PolygonVertexIDs.Num(), false );
+
+			// Map all of the polygon's vertex IDs to their extruded counterparts to create the new polygon perimeter
+			for( int32 PolygonVertexNumber = 0; PolygonVertexNumber < PolygonVertexIDs.Num(); ++PolygonVertexNumber )
+			{
+				const FVertexID VertexID = PolygonVertexIDs[ PolygonVertexNumber ];
+				NewFrontPolygonVertices[ PolygonVertexNumber ].VertexID = VertexIDToExtrudedCopyMap[ VertexID ];
+			}
+
+			{
+				static TArray<FPolygonToCreate> PolygonsToCreate;
+				PolygonsToCreate.Reset();
+
+				// Create the polygon
+				// @todo mesheditor perf: Ideally we support creating multiple polygons at once and batching up the work
+				FPolygonToCreate& PolygonToCreate = *new( PolygonsToCreate ) FPolygonToCreate();
+				PolygonToCreate.SectionID = PolygonRef.SectionID;
+				PolygonToCreate.PerimeterVertices = NewFrontPolygonVertices;	// @todo mesheditor perf: Copying static array here, ideally allocations could be avoided
+																				// @todo mesheditor holes: This algorithm is ignoring polygon holes!  Should be easy to support this by doing roughly the same thing.
+				static TArray<FPolygonRef> NewPolygonRefs;
+				NewPolygonRefs.Reset();
+				static TArray<FEdgeID> NewEdgeIDs;
+				NewEdgeIDs.Reset();
+				CreatePolygons( PolygonsToCreate, /* Out */ NewPolygonRefs, /* Out */ NewEdgeIDs );	// @todo mesheditor perf: Extra allocatons/copies: Ideally MoveTemp() here but we can't move a STATIC local!
+
+				ExtrudedFrontPolygonRef = NewPolygonRefs[ 0 ];
+			}
+			AllNewPolygons.Add( ExtrudedFrontPolygonRef );
+
+			// Retain the UVs from the original polygon
+			{
+				FVertexAttributesForPolygon& PolygonNewAttributes = *new( VertexAttributesForPolygons ) FVertexAttributesForPolygon();
+				PolygonNewAttributes.PolygonRef = ExtrudedFrontPolygonRef;
+				PolygonNewAttributes.PerimeterVertexAttributeLists.SetNum( PolygonVertexIDs.Num(), false );
+
+				for( int32 PolygonVertexNumber = 0; PolygonVertexNumber < PolygonVertexIDs.Num(); ++PolygonVertexNumber )
+				{
+					TArray<FMeshElementAttributeData>& PerimeterVertexNewAttributes = PolygonNewAttributes.PerimeterVertexAttributeLists[ PolygonVertexNumber ].Attributes;
+					for( int32 TextureCoordinateIndex = 0; TextureCoordinateIndex < TextureCoordinateCount; ++TextureCoordinateIndex )
+					{
+						const FVector4 TextureCoordinate = this->GetPolygonPerimeterVertexAttribute(
+							PolygonRef,
+							PolygonVertexNumber,
+							UEditableMeshAttribute::VertexTextureCoordinate(),
+							TextureCoordinateIndex );
+
+						FMeshElementAttributeData& PerimeterVertexAttribute = *new( PerimeterVertexNewAttributes ) FMeshElementAttributeData( UEditableMeshAttribute::VertexTextureCoordinate(), TextureCoordinateIndex, TextureCoordinate );
+					}
+
+					const FVector4 VertexColor = this->GetPolygonPerimeterVertexAttribute(
+						PolygonRef,
+						PolygonVertexNumber,
+						UEditableMeshAttribute::VertexColor(),
+						0 );
+
+					FMeshElementAttributeData& PerimeterVertexAttribute = *new( PerimeterVertexNewAttributes ) FMeshElementAttributeData( UEditableMeshAttribute::VertexColor(), 0, VertexColor );
+				}
+			}
+
+
+			// All of this edges of the new polygon will be hard
+			{
+				const int32 PerimeterEdgeCount = this->GetPolygonPerimeterEdgeCount( ExtrudedFrontPolygonRef );
+				for( int32 PerimeterEdgeNumber = 0; PerimeterEdgeNumber < PerimeterEdgeCount; ++PerimeterEdgeNumber )
+				{
+					bool bEdgeWindingIsReversedForPolygon;
+					const FEdgeID EdgeID = this->GetPolygonPerimeterEdge( ExtrudedFrontPolygonRef, PerimeterEdgeNumber, /* Out */ bEdgeWindingIsReversedForPolygon );
+
+					FAttributesForEdge& AttributesForEdge = *new( AttributesForEdges ) FAttributesForEdge();
+					AttributesForEdge.EdgeID = EdgeID;
+					AttributesForEdge.EdgeAttributes.Attributes.Add( FMeshElementAttributeData( UEditableMeshAttribute::EdgeIsHard(), 0, FVector4( 1.0f ) ) );
+					AttributesForEdge.EdgeAttributes.Attributes.Add( FMeshElementAttributeData( UEditableMeshAttribute::EdgeCreaseSharpness(), 0, FVector4( 1.0f ) ) );
+				}
+			}
+		}
+
+		OutNewExtrudedFrontPolygons.Add( ExtrudedFrontPolygonRef );
+	}
+
+	// Update edge attributes in bulk
+	SetEdgesAttributes( AttributesForEdges );
+
+	// Update vertex attributes in bulk
+	SetVerticesAttributes( AttributesForVertices );
+	SetPolygonsVertexAttributes( VertexAttributesForPolygons );
+
+	// Delete the original polygons
+	{
+		const bool bDeleteOrphanedEdges = true;
+		const bool bDeleteOrphanedVertices = true;
+		const bool bDeleteEmptySections = false;
+		DeletePolygons( Polygons, bDeleteOrphanedEdges, bDeleteOrphanedVertices, bDeleteEmptySections );
+	}
+
+	// Generate normals for everything, including adjacents because we might have changed edge hardness
+	GenerateNormalsAndTangentsForPolygonsAndAdjacents( AllNewPolygons );
+}
+
+
+void UEditableMesh::ExtendEdges( const TArray<FEdgeID>& EdgeIDs, const bool bWeldNeighbors, TArray<FEdgeID>& OutNewExtendedEdgeIDs )
+{
+	OutNewExtendedEdgeIDs.Reset();
+
+	static TArray<FVertexID> NewVertexIDs;
+	NewVertexIDs.Reset();
+
+	// For each original edge vertex ID that we'll be creating a counterpart for on the extended edge, a mapping
+	// to the vertex number of our NewVertexIDs (and VerticesToCreate) list.
+	static TMap<FVertexID, int32> OriginalVertexIDToCreatedVertexNumber;
+	OriginalVertexIDToCreatedVertexNumber.Reset();
+
+	// Create new vertices for all of the new edges.  If bWeldNeighbors is true, we'll share vertices between edges that share the
+	// same vertex instead of creating new edges.
+	{
+		static TArray<FVertexToCreate> VerticesToCreate;
+		VerticesToCreate.Reset();
+		VerticesToCreate.Reserve( EdgeIDs.Num() * 2 );	// Might actually end up needing less, but that's OK.
+
+		for( const FEdgeID EdgeID : EdgeIDs )
+		{
+			FVertexID EdgeVertexIDs[ 2 ];
+			GetEdgeVertices( EdgeID, /* Out */ EdgeVertexIDs[0], /* Out */ EdgeVertexIDs[1] );
+
+			for( const FVertexID EdgeVertexID : EdgeVertexIDs )
+			{
+				// Have we already created a counterpart for this vertex?  If we were asked to weld extended neighbor edges,
+				// we'll want to make sure that we share the extended vertex too!
+				const int32* FoundCreatedVertexNumberPtr = OriginalVertexIDToCreatedVertexNumber.Find( EdgeVertexID );
+				if( !( bWeldNeighbors && FoundCreatedVertexNumberPtr != nullptr ) )
+				{
+					const int32 CreatedVertexNumber = VerticesToCreate.Num();
+					FVertexToCreate& VertexToCreate = *new( VerticesToCreate ) FVertexToCreate();
+
+					// Copy attributes from the original vertex
+					for( const FName AttributeName : UEditableMesh::GetValidVertexAttributes() )
+					{
+						const int32 MaxAttributeIndex = GetMaxAttributeIndex( AttributeName );
+						for( int32 AttributeIndex = 0; AttributeIndex < MaxAttributeIndex; ++AttributeIndex )
+						{
+							FMeshElementAttributeData& VertexAttribute = *new( VertexToCreate.VertexAttributes.Attributes ) FMeshElementAttributeData(
+								AttributeName,
+								AttributeIndex,
+								GetVertexAttribute( EdgeVertexID, AttributeName, AttributeIndex ) );
+						}
+					}
+
+					// Keep track of which vertex we're creating a counterpart for
+					// @todo mesheditor perf: In the case where bWeldNeighbors=false, we can skip updating this map and instead calculate
+					// the index of the new vertex (it will be 2*EdgeNumber+EdgeVertexNumber, because all vertices will be unique.)
+					OriginalVertexIDToCreatedVertexNumber.Add( EdgeVertexID, CreatedVertexNumber );
+				}
+			}
+		}
+
+		CreateVertices( VerticesToCreate, /* Out */ NewVertexIDs );
+	}
+
+
+	// Create the extended edges
+	{
+		static TArray<FEdgeToCreate> EdgesToCreate;
+		EdgesToCreate.Reset();
+		EdgesToCreate.Reserve( EdgeIDs.Num() );
+
+		for( const FEdgeID EdgeID : EdgeIDs )
+		{
+			FVertexID EdgeVertexIDs[ 2 ];
+			GetEdgeVertices( EdgeID, /* Out */ EdgeVertexIDs[0], /* Out */ EdgeVertexIDs[1] );
+
+			FEdgeToCreate& EdgeToCreate = *new( EdgesToCreate ) FEdgeToCreate();
+
+			EdgeToCreate.VertexID0 = NewVertexIDs[ OriginalVertexIDToCreatedVertexNumber.FindChecked( EdgeVertexIDs[ 0 ] ) ];
+			EdgeToCreate.VertexID1 = NewVertexIDs[ OriginalVertexIDToCreatedVertexNumber.FindChecked( EdgeVertexIDs[ 1 ] ) ];
+
+			// Copy attributes from our original edge
+			for( const FName AttributeName : UEditableMesh::GetValidEdgeAttributes() )
+			{
+				const int32 MaxAttributeIndex = GetMaxAttributeIndex( AttributeName );
+				for( int32 AttributeIndex = 0; AttributeIndex < MaxAttributeIndex; ++AttributeIndex )
+				{
+					FMeshElementAttributeData& EdgeAttribute = *new( EdgeToCreate.EdgeAttributes.Attributes ) FMeshElementAttributeData(
+						AttributeName,
+						AttributeIndex,
+						GetEdgeAttribute( EdgeID, AttributeName, AttributeIndex ) );
+				}
+			}
+
+			// We're not connected to any polygons yet.  That will come later.
+			EdgeToCreate.ConnectedPolygons.Reset();
+		}
+
+		CreateEdges( EdgesToCreate, /* Out */ OutNewExtendedEdgeIDs );
+	}
+
+
+	// For every edge, make a quad to connect the original edge with it's extended counterpart.
+	{
+		static TArray<FPolygonToCreate> PolygonsToCreate;
+		PolygonsToCreate.Reset();
+		PolygonsToCreate.Reserve( EdgeIDs.Num() );
+
+		for( int32 ExtendedEdgeNumber = 0; ExtendedEdgeNumber < OutNewExtendedEdgeIDs.Num(); ++ExtendedEdgeNumber )
+		{
+			const FEdgeID OriginalEdgeID = EdgeIDs[ ExtendedEdgeNumber ];
+			const FEdgeID ExtendedEdgeID = OutNewExtendedEdgeIDs[ ExtendedEdgeNumber ];
+
+			FVertexID OriginalEdgeVertexIDs[ 2 ];
+			GetEdgeVertices( OriginalEdgeID, /* Out */ OriginalEdgeVertexIDs[0], /* Out */ OriginalEdgeVertexIDs[1] );
+
+			FVertexID ExtendedEdgeVertexIDs[ 2 ];
+			GetEdgeVertices( ExtendedEdgeID, /* Out */ ExtendedEdgeVertexIDs[0], /* Out */ ExtendedEdgeVertexIDs[1] );
+
+			FPolygonToCreate& PolygonToCreate = *new( PolygonsToCreate ) FPolygonToCreate();
+
+			// We need to figure out which mesh section to put the new polygons in.  To do this, we'll look at
+			// which polygons are already connected to the current edge, and use section number from the first
+			// polygon we can find.  If no polygons are connected, then we'll just use first section in the mesh.
+			// We'll also capture texture coordinates from this polygon, so we can apply them to the new polygon vertices.
+			TOptional<FPolygonRef> ConnectedPolygonRef;
+			{
+				const int32 ConnectedPolygonCount = GetEdgeConnectedPolygonCount( OriginalEdgeID );
+				if( ConnectedPolygonCount > 0 )
+				{
+					ConnectedPolygonRef = GetEdgeConnectedPolygon( OriginalEdgeID, 0 );
+				}
+			}
+
+			PolygonToCreate.SectionID = ConnectedPolygonRef.IsSet() ? ConnectedPolygonRef->SectionID : GetFirstValidSection();
+			check( PolygonToCreate.SectionID != FSectionID::Invalid );
+
+			PolygonToCreate.PerimeterVertices.SetNum( 4, false );
+			FVertexID ConnectedPolygonVertexIDsForTextureCoordinates[ 4 ];
+			{
+				// @todo mesheditor urgent subdiv: This causes degenerate UV triangles, which OpenSubdiv sort of freaks out about (causes flickering geometry)
+				int32 NextVertexNumber = 0;
+				ConnectedPolygonVertexIDsForTextureCoordinates[ NextVertexNumber ] = OriginalEdgeVertexIDs[ 1 ];
+				PolygonToCreate.PerimeterVertices[ NextVertexNumber++ ].VertexID = OriginalEdgeVertexIDs[ 1 ];
+
+				ConnectedPolygonVertexIDsForTextureCoordinates[ NextVertexNumber ] = OriginalEdgeVertexIDs[ 0 ];
+				PolygonToCreate.PerimeterVertices[ NextVertexNumber++ ].VertexID = OriginalEdgeVertexIDs[ 0 ];
+
+				ConnectedPolygonVertexIDsForTextureCoordinates[ NextVertexNumber ] = OriginalEdgeVertexIDs[ 0 ];
+				PolygonToCreate.PerimeterVertices[ NextVertexNumber++ ].VertexID = ExtendedEdgeVertexIDs[ 0 ];
+
+				ConnectedPolygonVertexIDsForTextureCoordinates[ NextVertexNumber ] = OriginalEdgeVertexIDs[ 1 ];
+				PolygonToCreate.PerimeterVertices[ NextVertexNumber++ ].VertexID = ExtendedEdgeVertexIDs[ 1 ];
+			}
+
+			// Preserve polygon vertex attributes
+			if( ConnectedPolygonRef.IsSet() )
+			{
+				for( int32 PerimeterVertexNumber = 0; PerimeterVertexNumber < PolygonToCreate.PerimeterVertices.Num(); ++PerimeterVertexNumber )
+				{
+					const FVertexID ConnectedPolygonVertexIDForTextureCoordinates = ConnectedPolygonVertexIDsForTextureCoordinates[ PerimeterVertexNumber ];
+
+					TArray<FMeshElementAttributeData>& AttributesForVertex = PolygonToCreate.PerimeterVertices[ PerimeterVertexNumber ].PolygonVertexAttributes.Attributes;
+				
+					const int32 ConnectedPolygonPerimeterVertexNumber = FindPolygonPerimeterVertexNumberForVertex( ConnectedPolygonRef.GetValue(), ConnectedPolygonVertexIDForTextureCoordinates );
+					check( ConnectedPolygonPerimeterVertexNumber != INDEX_NONE );
+
+					for( int32 TextureCoordinateIndex = 0; TextureCoordinateIndex < TextureCoordinateCount; ++TextureCoordinateIndex )
+					{
+						const FVector4 TextureCoordinate = GetPolygonPerimeterVertexAttribute( ConnectedPolygonRef.GetValue(), ConnectedPolygonPerimeterVertexNumber, UEditableMeshAttribute::VertexTextureCoordinate(), TextureCoordinateIndex );
+						AttributesForVertex.Add( FMeshElementAttributeData( UEditableMeshAttribute::VertexTextureCoordinate(), TextureCoordinateIndex, TextureCoordinate ) );
+					}
+
+					const FVector4 VertexColor = GetPolygonPerimeterVertexAttribute( ConnectedPolygonRef.GetValue(), ConnectedPolygonPerimeterVertexNumber, UEditableMeshAttribute::VertexColor(), 0 );
+					AttributesForVertex.Add( FMeshElementAttributeData( UEditableMeshAttribute::VertexColor(), 0, VertexColor ) );
+				}
+			}
+		}
+
+		// Create the polygons.  Note that this will also automatically create the missing side edges that connect
+		// the original edge to it's extended counterpart.
+		static TArray<FPolygonRef> NewPolygonRefs;
+		NewPolygonRefs.Reset();
+		static TArray<FEdgeID> NewEdgeIDs;
+		NewEdgeIDs.Reset();
+		CreatePolygons( PolygonsToCreate, /* Out */ NewPolygonRefs, /* Out */ NewEdgeIDs );
+
+		// Expecting no more than two new edges to be created while creating the polygon.  It's possible for zero or one edge
+		// to be created, depending on how many edges we share with neighbors that were extended.
+		check( bWeldNeighbors ? ( NewEdgeIDs.Num() <= 2 * EdgeIDs.Num() ) : ( NewEdgeIDs.Num() == 2 * EdgeIDs.Num() ) );
+
+		// Generate normals for all of the new polygons
+		GenerateNormalsAndTangentsForPolygons( NewPolygonRefs );
+	}
+}
+
+
+void UEditableMesh::ExtendVertices( const TArray<FVertexID>& VertexIDs, const bool bOnlyExtendClosestEdge, const FVector ReferencePosition, TArray<FVertexID>& OutNewExtendedVertexIDs )
+{
+	OutNewExtendedVertexIDs.Reset();
+
+	// Create new vertices for all of the new edges.  If bWeldNeighbors is true, we'll share vertices between edges that share the
+	// same vertex instead of creating new edges.
+	{
+		static TArray<FVertexToCreate> VerticesToCreate;
+		VerticesToCreate.Reset();
+		VerticesToCreate.Reserve( VertexIDs.Num() );
+
+		for( const FVertexID VertexID : VertexIDs )
+		{
+			const int32 CreatedVertexNumber = VerticesToCreate.Num();
+			FVertexToCreate& VertexToCreate = *new( VerticesToCreate ) FVertexToCreate();
+
+			// Copy attributes from the original vertex
+			for( const FName AttributeName : UEditableMesh::GetValidVertexAttributes() )
+			{
+				const int32 MaxAttributeIndex = GetMaxAttributeIndex( AttributeName );
+				for( int32 AttributeIndex = 0; AttributeIndex < MaxAttributeIndex; ++AttributeIndex )
+				{
+					FMeshElementAttributeData& VertexAttribute = *new( VertexToCreate.VertexAttributes.Attributes ) FMeshElementAttributeData(
+						AttributeName,
+						AttributeIndex,
+						GetVertexAttribute( VertexID, AttributeName, AttributeIndex ) );
+				}
+			}
+		}
+
+		CreateVertices( VerticesToCreate, /* Out */ OutNewExtendedVertexIDs );
+	}
+
+
+	// For each vertex, we'll now create new triangles to connect the new vertex to each of the original vertex's adjacent vertices.
+	// If the option bOnlyExtendClosestEdge was enabled, we'll only bother doing this for next closest vertex (so, only a single
+	// triangle per vertex will be created.)
+	{
+		static TArray<FPolygonToCreate> PolygonsToCreate;
+		PolygonsToCreate.Reset();
+
+		for( int32 VertexNumber = 0; VertexNumber < VertexIDs.Num(); ++VertexNumber )
+		{
+			const FVertexID OriginalVertexID = VertexIDs[ VertexNumber ];
+			const FVertexID NewVertexID = OutNewExtendedVertexIDs[ VertexNumber ];
+
+			FVertexID ClosestVertexID = FVertexID::Invalid;
+			if( bOnlyExtendClosestEdge )
+			{
+				// Iterate over the edges connected to this vertex, and figure out which edge is closest to the 
+				// specified reference position
+				float ClosestSquaredEdgeDistance = TNumericLimits<float>::Max();
+
+				const int32 ConnectedEdgeCount = GetVertexConnectedEdgeCount( OriginalVertexID );
+				for( int32 EdgeNumber = 0; EdgeNumber < ConnectedEdgeCount; ++EdgeNumber )
+				{
+					const FEdgeID ConnectedEdgeID = GetVertexConnectedEdge( OriginalVertexID, EdgeNumber );
+
+					FVertexID EdgeVertexIDs[ 2 ];
+					GetEdgeVertices( ConnectedEdgeID, /* Out */ EdgeVertexIDs[0], /* Out */ EdgeVertexIDs[1] );
+
+					FVector EdgeVertexPositions[ 2 ];
+					for( int32 EdgeVertexNumber = 0; EdgeVertexNumber < 2; ++EdgeVertexNumber )
+					{
+						EdgeVertexPositions[ EdgeVertexNumber ] = GetVertexAttribute( EdgeVertexIDs[ EdgeVertexNumber ], UEditableMeshAttribute::VertexPosition(), 0 );
+					}
+
+					const float SquaredEdgeDistance = FMath::PointDistToSegmentSquared( ReferencePosition, EdgeVertexPositions[ 0 ], EdgeVertexPositions[ 1 ] );
+					if( SquaredEdgeDistance < ClosestSquaredEdgeDistance )
+					{
+						ClosestVertexID = EdgeVertexIDs[ 0 ] == OriginalVertexID ? EdgeVertexIDs[ 1 ] : EdgeVertexIDs[ 0 ];
+						ClosestSquaredEdgeDistance = SquaredEdgeDistance;
+					}
+				}
+			}
+
+			static TArray<FVertexID> AdjacentVertexIDs;
+			GetVertexAdjacentVertices( OriginalVertexID, /* Out */ AdjacentVertexIDs );
+
+			// For every adjacent vertex, go ahead and create a new triangle
+			for( const FVertexID AdjacentVertexID : AdjacentVertexIDs )
+			{
+				// If we were asked to only extend an edge that's closest to a reference position, check for that here
+				if( !bOnlyExtendClosestEdge || ( AdjacentVertexID == ClosestVertexID ) )
+				{
+					FPolygonToCreate& PolygonToCreate = *new( PolygonsToCreate ) FPolygonToCreate();
+
+					// Figure out which of the connected polygons shares the edge we're going to be using
+					TOptional<FPolygonRef> ConnectedPolygonRef;
+					{
+						static TArray<FPolygonRef> ConnectedPolygonRefs;
+						GetVertexConnectedPolygons( OriginalVertexID, /* Out */ ConnectedPolygonRefs );
+
+						for( const FPolygonRef PolygonRef : ConnectedPolygonRefs )
+						{
+							const int32 AdjacentVertexNumber = FindPolygonPerimeterVertexNumberForVertex( PolygonRef, AdjacentVertexID );
+							if( AdjacentVertexNumber != INDEX_NONE )
+							{
+								// NOTE: There can be more than one polygon that shares this edge.  We'll just take the first one we find.
+								ConnectedPolygonRef = PolygonRef;
+								break;
+							}
+						}
+					}
+
+					bool bConnectedPolygonWindsForward = true;
+					if( ConnectedPolygonRef.IsSet() )
+					{
+						const int32 OriginalVertexNumberOnConnectedPolygon = FindPolygonPerimeterVertexNumberForVertex( ConnectedPolygonRef.GetValue(), OriginalVertexID );
+						check( OriginalVertexNumberOnConnectedPolygon != INDEX_NONE );
+
+						const int32 AdjacentVertexNumberOnConnectedPolygon = FindPolygonPerimeterVertexNumberForVertex( ConnectedPolygonRef.GetValue(), AdjacentVertexID );
+						check( AdjacentVertexNumberOnConnectedPolygon != INDEX_NONE );
+
+						const int32 PerimeterVertexCount = GetPolygonPerimeterVertexCount( ConnectedPolygonRef.GetValue() );
+						if( !( OriginalVertexNumberOnConnectedPolygon == ( PerimeterVertexCount - 1 ) && AdjacentVertexNumberOnConnectedPolygon == 0 ) &&  // Doesn't wrap forward?
+							( OriginalVertexNumberOnConnectedPolygon > AdjacentVertexNumberOnConnectedPolygon ||			// Winds backwards?
+							  AdjacentVertexNumberOnConnectedPolygon == ( PerimeterVertexCount - 1 ) && OriginalVertexNumberOnConnectedPolygon == 0 ) )	// Wraps backwards?
+						{
+							bConnectedPolygonWindsForward = false;
+						}
+					}
+
+					PolygonToCreate.SectionID = ConnectedPolygonRef.IsSet() ? ConnectedPolygonRef->SectionID : GetFirstValidSection();
+					check( PolygonToCreate.SectionID != FSectionID::Invalid );
+
+					FVertexID ConnectedPolygonVertexIDsForTextureCoordinates[ 3 ];
+					PolygonToCreate.PerimeterVertices.SetNum( 3, false );
+					{
+						int32 NextVertexNumber = 0;
+
+						// Original selected vertex
+						ConnectedPolygonVertexIDsForTextureCoordinates[ NextVertexNumber ] = OriginalVertexID;
+						PolygonToCreate.PerimeterVertices[ NextVertexNumber++ ].VertexID = OriginalVertexID;
+
+						if( bConnectedPolygonWindsForward )
+						{
+							// The new vertex we created
+							ConnectedPolygonVertexIDsForTextureCoordinates[ NextVertexNumber ] = OriginalVertexID;
+							PolygonToCreate.PerimeterVertices[ NextVertexNumber++ ].VertexID = NewVertexID;
+
+							// The adjacent vertex
+							ConnectedPolygonVertexIDsForTextureCoordinates[ NextVertexNumber ] = AdjacentVertexID;
+							PolygonToCreate.PerimeterVertices[ NextVertexNumber++ ].VertexID = AdjacentVertexID;
+						}
+						else
+						{
+							// The adjacent vertex
+							ConnectedPolygonVertexIDsForTextureCoordinates[ NextVertexNumber ] = AdjacentVertexID;
+							PolygonToCreate.PerimeterVertices[ NextVertexNumber++ ].VertexID = AdjacentVertexID;
+
+							// The new vertex we created
+							ConnectedPolygonVertexIDsForTextureCoordinates[ NextVertexNumber ] = OriginalVertexID;
+							PolygonToCreate.PerimeterVertices[ NextVertexNumber++ ].VertexID = NewVertexID;
+						}
+					}
+
+					// Preserve polygon vertex attributes
+					if( ConnectedPolygonRef.IsSet() )
+					{
+						for( int32 PerimeterVertexNumber = 0; PerimeterVertexNumber < PolygonToCreate.PerimeterVertices.Num(); ++PerimeterVertexNumber )
+						{
+							const FVertexID ConnectedPolygonVertexIDForTextureCoordinates = ConnectedPolygonVertexIDsForTextureCoordinates[ PerimeterVertexNumber ];
+
+							TArray<FMeshElementAttributeData>& AttributesForVertex = PolygonToCreate.PerimeterVertices[ PerimeterVertexNumber ].PolygonVertexAttributes.Attributes;
+
+							const int32 ConnectedPolygonPerimeterVertexNumber = FindPolygonPerimeterVertexNumberForVertex( ConnectedPolygonRef.GetValue(), ConnectedPolygonVertexIDForTextureCoordinates );
+							check( ConnectedPolygonPerimeterVertexNumber != INDEX_NONE );
+
+							for( int32 TextureCoordinateIndex = 0; TextureCoordinateIndex < TextureCoordinateCount; ++TextureCoordinateIndex )
+							{
+								const FVector4 TextureCoordinate = GetPolygonPerimeterVertexAttribute( ConnectedPolygonRef.GetValue(), ConnectedPolygonPerimeterVertexNumber, UEditableMeshAttribute::VertexTextureCoordinate(), TextureCoordinateIndex );
+								AttributesForVertex.Add( FMeshElementAttributeData( UEditableMeshAttribute::VertexTextureCoordinate(), TextureCoordinateIndex, TextureCoordinate ) );
+							}
+
+							const FVector4 VertexColor = GetPolygonPerimeterVertexAttribute( ConnectedPolygonRef.GetValue(), ConnectedPolygonPerimeterVertexNumber, UEditableMeshAttribute::VertexColor(), 0 );
+							AttributesForVertex.Add( FMeshElementAttributeData( UEditableMeshAttribute::VertexColor(), 0, VertexColor ) );
+						}
+					}
+				}
+			}		
+		}
+
+		static TArray<FPolygonRef> NewPolygonRefs;
+		NewPolygonRefs.Reset();
+
+		// Create the polygons.  Note that this will also automatically create the missing side edges that connect
+		// the original edge to it's extended counterpart.
+		static TArray<FEdgeID> NewEdgeIDs;
+		NewEdgeIDs.Reset();
+		CreatePolygons( PolygonsToCreate, /* Out */ NewPolygonRefs, /* Out */ NewEdgeIDs );
+
+		// Generate normals for all of the new polygons
+		GenerateNormalsAndTangentsForPolygons( NewPolygonRefs );
+	}
+}
+
+
+void UEditableMesh::ComputePolygonsSharedEdges( const TArray<FPolygonRef>& PolygonRefs, TArray<FEdgeID>& OutSharedEdgeIDs ) const
+{
+	OutSharedEdgeIDs.Reset();
+
+	static TSet<FEdgeID> EdgesSeenSoFar;
+	EdgesSeenSoFar.Reset();
+	for( const FPolygonRef PolygonRef : PolygonRefs )
+	{
+		static TArray<FEdgeID> PolygonPerimeterEdgeIDs;
+		GetPolygonPerimeterEdges( PolygonRef, /* Out */ PolygonPerimeterEdgeIDs );
+
+		for( const FEdgeID EdgeID : PolygonPerimeterEdgeIDs )
+		{
+			bool bWasAlreadyInSet = false;
+			EdgesSeenSoFar.Add( EdgeID, /* Out */ &bWasAlreadyInSet );
+
+			if( bWasAlreadyInSet )
+			{
+				// OK, this edge was referenced by more than one polygon!
+				OutSharedEdgeIDs.Add( EdgeID );
+			}
+		}
+	}
+}
+
+
+void UEditableMesh::BevelOrInsetPolygons( const TArray<FPolygonRef>& PolygonRefs, const float InsetFixedDistance, const float InsetProgressTowardCenter, const EInsetPolygonsMode Mode, const bool bShouldBevel, TArray<FPolygonRef>& OutNewCenterPolygonRefs, TArray<FPolygonRef>& OutNewSidePolygonRefs )
+{
+	// @todo mesheditor inset/bevel: Weird feedback loop issues at glancing angles when moving mouse while beveling
+
+	static TArray<FPolygonToCreate> SidePolygonsToCreate;
+	SidePolygonsToCreate.Reset();
+
+	static TArray<FPolygonToCreate> CenterPolygonsToCreate;
+	CenterPolygonsToCreate.Reset();
+
+	static TArray<FAttributesForVertex> AttributesForVertices;
+	AttributesForVertices.Reset();
+
+	for( const FPolygonRef PolygonRef : PolygonRefs )
+	{
+		// Find the center of this polygon
+		const FVector PolygonCenter = ComputePolygonCenter( PolygonRef );
+
+		static TArray<FVertexID> PerimeterVertexIDs;
+		GetPolygonPerimeterVertices( PolygonRef, /* Out */ PerimeterVertexIDs );
+
+		static TArray<FVertexToCreate> VerticesToCreate;
+		VerticesToCreate.Reset( PerimeterVertexIDs.Num() );
+
+		static TArray<TArray<FVector4>> TextureCoordinatesForNewVertices;	// @todo mesheditor perf: Extra memory alloc every time because of nested array
+		TextureCoordinatesForNewVertices.Reset();
+		TextureCoordinatesForNewVertices.Reserve( PerimeterVertexIDs.Num() );
+
+		static TArray<FVector4> VertexColorsForNewVertices;
+		VertexColorsForNewVertices.Reset();
+		VertexColorsForNewVertices.Reserve( PerimeterVertexIDs.Num() );
+
+		for( int32 PerimeterVertexNumber = 0; PerimeterVertexNumber < PerimeterVertexIDs.Num(); ++PerimeterVertexNumber )
+		{
+			const FVertexID PerimeterVertexID = PerimeterVertexIDs[ PerimeterVertexNumber ];
+
+			FVector OffsetDirection = FVector::ZeroVector;
+
+			// If we're beveling, go ahead and move the original polygon perimeter vertices
+			if( bShouldBevel )
+			{
+				// Figure out if this vertex is shared with other polygons that we were asked to bevel.  If it is,
+				// then we'll want to offset the vertex in the average direction of all of those shared polygons.
+				// However, if the vertex is ONLY shared with polygons we were asked to bevel (no other polygons),
+				// then we don't need to move it at all -- it's an internal edge vertex.
+				// @todo mesheditor bevel: With multiple polygons, we don't want to move vertices that are part of
+				// a shared internal border edge.  But we're not handling that yet.
+				int32 ConnectedBevelPolygonCount = 0;
+				bool bIsOnlyConnectedToBevelPolygons = true;
+
+				static TArray<FPolygonRef> ConnectedPolygonRefs;
+				GetVertexConnectedPolygons( PerimeterVertexID, /* Out */ ConnectedPolygonRefs );
+				for( const FPolygonRef ConnectedPolygonRef : ConnectedPolygonRefs )
+				{
+					if( PolygonRefs.Contains( ConnectedPolygonRef ) )
+					{
+						++ConnectedBevelPolygonCount;
+						const FVector ConnectedPolygonNormal = ComputePolygonNormal( ConnectedPolygonRef );
+						OffsetDirection += -ConnectedPolygonNormal;
+					}
+					else
+					{
+						bIsOnlyConnectedToBevelPolygons = false;
+					}
+				}
+
+				OffsetDirection.Normalize();
+			}
+
+
+			const FVector VertexPosition = GetVertexAttribute( PerimeterVertexID, UEditableMeshAttribute::VertexPosition(), 0 );
+
+			FVector DirectionTowardCenter;
+			float DistanceToCenter;
+			( PolygonCenter - VertexPosition ).ToDirectionAndLength( /* Out */ DirectionTowardCenter, /* Out */ DistanceToCenter );
+
+			const float InsetOffset = ( DistanceToCenter * InsetProgressTowardCenter + InsetFixedDistance );
+			const FVector InsetVertexPosition = VertexPosition + DirectionTowardCenter * InsetOffset;
+
+			FVertexToCreate& VertexToCreate = *new( VerticesToCreate ) FVertexToCreate();
+			VertexToCreate.VertexAttributes.Attributes.Add( FMeshElementAttributeData( UEditableMeshAttribute::VertexPosition(), 0, InsetVertexPosition ) );
+
+			// Interpolate the texture coordinates and vertex colors on this polygon to figure out UVs for our new vertex
+			static TArray<int32> PerimeterVertexIndices;
+			FVector TriangleVertexWeights;
+
+			if( ComputeBarycentricWeightForPointOnPolygon( PolygonRef, InsetVertexPosition, PerimeterVertexIndices, TriangleVertexWeights ) )
+			{
+				TArray<FVector4>& TextureCoordinatesForNewVertex = *new( TextureCoordinatesForNewVertices ) TArray<FVector4>();
+				TextureCoordinatesForNewVertex.SetNum( TextureCoordinateCount, false );
+
+				for( int32 TextureCoordinateIndex = 0; TextureCoordinateIndex < TextureCoordinateCount; ++TextureCoordinateIndex )
+				{
+					TextureCoordinatesForNewVertex[ TextureCoordinateIndex ] =
+						TriangleVertexWeights.X * GetPolygonPerimeterVertexAttribute( PolygonRef, PerimeterVertexIndices[ 0 ], UEditableMeshAttribute::VertexTextureCoordinate(), TextureCoordinateIndex ) +
+						TriangleVertexWeights.Y * GetPolygonPerimeterVertexAttribute( PolygonRef, PerimeterVertexIndices[ 1 ], UEditableMeshAttribute::VertexTextureCoordinate(), TextureCoordinateIndex ) +
+						TriangleVertexWeights.Z * GetPolygonPerimeterVertexAttribute( PolygonRef, PerimeterVertexIndices[ 2 ], UEditableMeshAttribute::VertexTextureCoordinate(), TextureCoordinateIndex );
+				}
+
+				VertexColorsForNewVertices.Add(
+					TriangleVertexWeights.X * GetPolygonPerimeterVertexAttribute( PolygonRef, PerimeterVertexIndices[ 0 ], UEditableMeshAttribute::VertexColor(), 0 ) +
+					TriangleVertexWeights.Y * GetPolygonPerimeterVertexAttribute( PolygonRef, PerimeterVertexIndices[ 1 ], UEditableMeshAttribute::VertexColor(), 0 ) +
+					TriangleVertexWeights.Z * GetPolygonPerimeterVertexAttribute( PolygonRef, PerimeterVertexIndices[ 2 ], UEditableMeshAttribute::VertexColor(), 0 )
+				);
+			}
+			else
+			{
+				TArray<FVector4>& TextureCoordinatesForNewVertex = *new( TextureCoordinatesForNewVertices ) TArray<FVector4>();
+
+				// The new vertex must have been slightly outside the bounds of the polygon, probably just floating point error.
+				// Just use zero'd out texture coordinates and white vertex color in this case.
+				TextureCoordinatesForNewVertex.SetNumZeroed( TextureCoordinateCount );
+				VertexColorsForNewVertices.Emplace( 1.0f, 1.0f, 1.0f, 1.0f );
+			}
+
+			// If we're beveling, go ahead and move the original polygon perimeter vertices
+			if( bShouldBevel )
+			{
+				// Offset the vertex by the opposite direction of the polygon's normal.  We'll move it the same distance
+				// that we're insetting the new polygon
+				const FVector NewVertexPosition = VertexPosition + OffsetDirection * InsetOffset;
+
+				// @todo mesheditor urgent: Undo/Redo breaks if we have the same vertex more than once for some reason!!  Get rid of this ideally
+				const bool bAlreadyHaveVertex = AttributesForVertices.ContainsByPredicate( [PerimeterVertexID]( FAttributesForVertex& AttributesForVertex ) { return AttributesForVertex.VertexID == PerimeterVertexID; } );
+				if( !bAlreadyHaveVertex )
+				{
+					FAttributesForVertex& AttributesForVertex = *new( AttributesForVertices ) FAttributesForVertex();
+					AttributesForVertex.VertexID = PerimeterVertexID;
+
+					FMeshElementAttributeData& ElementAttributeData = *new( AttributesForVertex.VertexAttributes.Attributes ) FMeshElementAttributeData();
+					ElementAttributeData.AttributeName = UEditableMeshAttribute::VertexPosition();
+					ElementAttributeData.AttributeValue = NewVertexPosition;
+					ElementAttributeData.AttributeIndex = 0;
+				}
+			}
+		}
+
+		static TArray<FVertexID> NewVertexIDs;
+		CreateVertices( VerticesToCreate, /* Out */ NewVertexIDs );
+
+		// The new (inset) polygon will be surrounded by new "side" quad polygons, one for each vertex of the perimeter
+		// that's being inset.
+		if( Mode == EInsetPolygonsMode::All || Mode == EInsetPolygonsMode::SidePolygonsOnly )
+		{
+			const int32 NewSidePolygonCount = NewVertexIDs.Num();
+			for( int32 SidePolygonNumber = 0; SidePolygonNumber < NewSidePolygonCount; ++SidePolygonNumber )
+			{
+				const int32 LeftVertexNumber = SidePolygonNumber;
+				const int32 RightVertexNumber = ( LeftVertexNumber + 1 ) % NewSidePolygonCount;
+
+				const FVertexID LeftOriginalVertexID = PerimeterVertexIDs[ LeftVertexNumber ];
+				const FVertexID LeftInsetVertexID = NewVertexIDs[ LeftVertexNumber ];
+				const FVertexID RightOriginalVertexID = PerimeterVertexIDs[ RightVertexNumber ];
+				const FVertexID RightInsetVertexID = NewVertexIDs[ RightVertexNumber ];
+
+				FPolygonToCreate& PolygonToCreate = *new( SidePolygonsToCreate ) FPolygonToCreate();
+				PolygonToCreate.SectionID = PolygonRef.SectionID;
+
+				// Original left
+				{
+					FVertexAndAttributes& NewPerimeterVertex = *new( PolygonToCreate.PerimeterVertices ) FVertexAndAttributes();
+					NewPerimeterVertex.VertexID = LeftOriginalVertexID;
+
+					// Copy all of the polygon vertex attributes over
+					for( int32 TextureCoordinateIndex = 0; TextureCoordinateIndex < TextureCoordinateCount; ++TextureCoordinateIndex )
+					{
+						NewPerimeterVertex.PolygonVertexAttributes.Attributes.Add( FMeshElementAttributeData(
+							UEditableMeshAttribute::VertexTextureCoordinate(),
+							TextureCoordinateIndex,
+							GetPolygonPerimeterVertexAttribute( PolygonRef, LeftVertexNumber, UEditableMeshAttribute::VertexTextureCoordinate(), TextureCoordinateIndex ) ) );
+					}
+
+					NewPerimeterVertex.PolygonVertexAttributes.Attributes.Add( FMeshElementAttributeData(
+						UEditableMeshAttribute::VertexColor(),
+						0,
+						GetPolygonPerimeterVertexAttribute( PolygonRef, LeftVertexNumber, UEditableMeshAttribute::VertexColor(), 0 ) ) );
+				}
+
+				// Original right
+				{
+					FVertexAndAttributes& NewPerimeterVertex = *new( PolygonToCreate.PerimeterVertices ) FVertexAndAttributes();
+					NewPerimeterVertex.VertexID = RightOriginalVertexID;
+
+					// Copy all of the polygon vertex attributes over
+					for( int32 TextureCoordinateIndex = 0; TextureCoordinateIndex < TextureCoordinateCount; ++TextureCoordinateIndex )
+					{
+						NewPerimeterVertex.PolygonVertexAttributes.Attributes.Add( FMeshElementAttributeData(
+							UEditableMeshAttribute::VertexTextureCoordinate(),
+							TextureCoordinateIndex,
+							GetPolygonPerimeterVertexAttribute( PolygonRef, RightVertexNumber, UEditableMeshAttribute::VertexTextureCoordinate(), TextureCoordinateIndex ) ) );
+					}
+
+					NewPerimeterVertex.PolygonVertexAttributes.Attributes.Add( FMeshElementAttributeData(
+						UEditableMeshAttribute::VertexColor(),
+						0,
+						GetPolygonPerimeterVertexAttribute( PolygonRef, RightVertexNumber, UEditableMeshAttribute::VertexColor(), 0 ) ) );
+				}
+
+				// Inset right
+				{
+					FVertexAndAttributes& NewPerimeterVertex = *new( PolygonToCreate.PerimeterVertices ) FVertexAndAttributes();
+					NewPerimeterVertex.VertexID = RightInsetVertexID;
+
+					// Use interpolated texture coordinates
+					for( int32 TextureCoordinateIndex = 0; TextureCoordinateIndex < TextureCoordinateCount; ++TextureCoordinateIndex )
+					{
+						NewPerimeterVertex.PolygonVertexAttributes.Attributes.Add( FMeshElementAttributeData(
+							UEditableMeshAttribute::VertexTextureCoordinate(),
+							TextureCoordinateIndex,
+							TextureCoordinatesForNewVertices[ RightVertexNumber ][ TextureCoordinateIndex ] ) );
+					}
+
+					NewPerimeterVertex.PolygonVertexAttributes.Attributes.Add( FMeshElementAttributeData(
+						UEditableMeshAttribute::VertexColor(),
+						0,
+						VertexColorsForNewVertices[ RightVertexNumber ] ) );
+				}
+
+				// Inset left
+				{
+					FVertexAndAttributes& NewPerimeterVertex = *new( PolygonToCreate.PerimeterVertices ) FVertexAndAttributes();
+					NewPerimeterVertex.VertexID = LeftInsetVertexID;
+
+					// Use interpolated texture coordinates
+					for( int32 TextureCoordinateIndex = 0; TextureCoordinateIndex < TextureCoordinateCount; ++TextureCoordinateIndex )
+					{
+						NewPerimeterVertex.PolygonVertexAttributes.Attributes.Add( FMeshElementAttributeData(
+							UEditableMeshAttribute::VertexTextureCoordinate(),
+							TextureCoordinateIndex,
+							TextureCoordinatesForNewVertices[ LeftVertexNumber ][ TextureCoordinateIndex ] ) );
+					}
+
+					NewPerimeterVertex.PolygonVertexAttributes.Attributes.Add( FMeshElementAttributeData(
+						UEditableMeshAttribute::VertexColor(),
+						0,
+						VertexColorsForNewVertices[ LeftVertexNumber ] ) );
+				}
+			}
+		}
+
+		// Now create the new center polygon that will connect all of the new inset vertices
+		if( Mode == EInsetPolygonsMode::All || Mode == EInsetPolygonsMode::CenterPolygonOnly )
+		{
+			FPolygonToCreate& PolygonToCreate = *new( CenterPolygonsToCreate ) FPolygonToCreate();
+			PolygonToCreate.SectionID = PolygonRef.SectionID;
+
+			for( int32 NewVertexNumber = 0; NewVertexNumber < NewVertexIDs.Num(); ++NewVertexNumber )
+			{
+				const FVertexID NewVertexID = NewVertexIDs[ NewVertexNumber ];
+
+				FVertexAndAttributes& NewPerimeterVertex = *new( PolygonToCreate.PerimeterVertices ) FVertexAndAttributes();
+				NewPerimeterVertex.VertexID = NewVertexID;
+
+				// Use interpolated texture coordinates
+				for( int32 TextureCoordinateIndex = 0; TextureCoordinateIndex < TextureCoordinateCount; ++TextureCoordinateIndex )
+				{
+					NewPerimeterVertex.PolygonVertexAttributes.Attributes.Add( FMeshElementAttributeData(
+						UEditableMeshAttribute::VertexTextureCoordinate(),
+						TextureCoordinateIndex,
+						TextureCoordinatesForNewVertices[ NewVertexNumber ][ TextureCoordinateIndex ] ) );
+				}
+
+				NewPerimeterVertex.PolygonVertexAttributes.Attributes.Add( FMeshElementAttributeData(
+					UEditableMeshAttribute::VertexColor(),
+					0,
+					VertexColorsForNewVertices[ NewVertexNumber ] ) );
+			}
+		}
+	}
+
+	// Delete the original polygons
+	{
+		const bool bDeleteOrphanedEdges = false;	// No need to delete orphans, because this function won't orphan anything
+		const bool bDeleteOrphanedVertices = false;
+		const bool bDeleteEmptySections = false;
+		DeletePolygons( PolygonRefs, bDeleteOrphanedEdges, bDeleteOrphanedVertices, bDeleteEmptySections );
+	}
+
+	// @todo mesheditor inset: Should we create the new inset polygon with hard edges or soft?  Currently using the default.
+
+	// Updated any vertices that need to be moved
+	if( AttributesForVertices.Num() > 0 )
+	{
+		SetVerticesAttributes( AttributesForVertices );
+	}
+
+	// Create the new side polygons
+	static TArray<FPolygonRef> AllNewPolygonRefs;
+	AllNewPolygonRefs.Reset();
+
+	static TArray<FEdgeID> EdgesToMakeHard;
+	EdgesToMakeHard.Reset();
+
+	static TArray<bool> EdgesNewIsHard;
+	EdgesNewIsHard.Reset();
+
+	static TArray<float> EdgesNewCreaseSharpness;
+	EdgesNewCreaseSharpness.Reset();
+
+	if( SidePolygonsToCreate.Num() > 0 )
+	{
+		static TArray<FEdgeID> NewEdgeIDs;
+		CreatePolygons( SidePolygonsToCreate, /* Out */ OutNewSidePolygonRefs, /* Out */ NewEdgeIDs );
+
+		AllNewPolygonRefs.Append( OutNewSidePolygonRefs );
+
+		// If we're beveling, go ahead and make all of the side edges hard
+		if( bShouldBevel )
+		{
+			static TSet<FEdgeID> AllSidePolygonEdgeIDs;
+			AllSidePolygonEdgeIDs.Reset();
+			for( const FPolygonRef PolygonRef : OutNewSidePolygonRefs )
+			{
+				static TArray<FEdgeID> EdgeIDs;
+				GetPolygonPerimeterEdges( PolygonRef, /* Out */ EdgeIDs );
+
+				AllSidePolygonEdgeIDs.Append( EdgeIDs );
+			}
+
+			for( const FEdgeID EdgeID : AllSidePolygonEdgeIDs )
+			{
+				EdgesToMakeHard.Add( EdgeID );
+				EdgesNewIsHard.Add( true );
+				EdgesNewCreaseSharpness.Add( 1.0f );
+			}
+		}
+	}
+
+	// Create the new center polygons.  Note that we pass back the IDs of the newly-created center polygons
+	if( CenterPolygonsToCreate.Num() > 0 )
+	{
+		static TArray<FEdgeID> NewEdgeIDs;
+		CreatePolygons( CenterPolygonsToCreate, /* Out */ OutNewCenterPolygonRefs, /* Out */ NewEdgeIDs );
+
+		AllNewPolygonRefs.Append( OutNewCenterPolygonRefs );
+
+		// Make the perimeter edges of the new polygons hard
+		static TSet<FEdgeID> PerimeterEdgeIDsToMakeHard;
+		PerimeterEdgeIDsToMakeHard.Reset();
+		for( const FPolygonRef NewCenterPolygonRef : OutNewCenterPolygonRefs )
+		{
+			static TArray<FEdgeID> PerimeterEdgeIDs;
+			GetPolygonPerimeterEdges( NewCenterPolygonRef, /* Out */ PerimeterEdgeIDs );
+
+			for( const FEdgeID EdgeID : PerimeterEdgeIDs )
+			{
+				PerimeterEdgeIDsToMakeHard.Add( EdgeID );
+			}
+		}
+
+		for( const FEdgeID EdgeID : PerimeterEdgeIDsToMakeHard )
+		{
+			EdgesToMakeHard.Add( EdgeID );
+			EdgesNewIsHard.Add( true );
+			EdgesNewCreaseSharpness.Add( 1.0f );
+		}
+	}
+
+	if( EdgesToMakeHard.Num() > 0 )
+	{
+		SetEdgesHardness( EdgesToMakeHard, EdgesNewIsHard );
+		SetEdgesCreaseSharpness( EdgesToMakeHard, EdgesNewCreaseSharpness );
+	}
+
+	// Generate normals for all of the new polygons
+	if( bShouldBevel )
+	{
+		// When beveling, we're moving the polygon perimeter vertex locations which will certainly
+		// affect the normals and tangents of adjacent polygons as well!
+		GenerateNormalsAndTangentsForPolygonsAndAdjacents( AllNewPolygonRefs );
+	}
+	else
+	{
+		// When not beveling, the original polygon's perimeter stays the same, so we don't need to
+		// worry about normals for adjacent polygons
+		GenerateNormalsAndTangentsForPolygons( AllNewPolygonRefs );
+	}
+}
+
+
+void UEditableMesh::InsetPolygons( const TArray<FPolygonRef>& PolygonRefs, const float InsetFixedDistance, const float InsetProgressTowardCenter, const EInsetPolygonsMode Mode, TArray<FPolygonRef>& OutNewCenterPolygonRefs, TArray<FPolygonRef>& OutNewSidePolygonRefs )
+{
+	const bool bShouldBevel = false;
+	BevelOrInsetPolygons( PolygonRefs, InsetFixedDistance, InsetProgressTowardCenter, Mode, bShouldBevel, /* Out */ OutNewCenterPolygonRefs, /* Out */ OutNewSidePolygonRefs );
+}
+
+
+void UEditableMesh::BevelPolygons( const TArray<FPolygonRef>& PolygonRefs, const float BevelFixedDistance, const float BevelProgressTowardCenter, TArray<FPolygonRef>& OutNewCenterPolygonRefs, TArray<FPolygonRef>& OutNewSidePolygonRefs )
+{
+	const bool bShouldBevel = true;
+	BevelOrInsetPolygons( PolygonRefs, BevelFixedDistance, BevelProgressTowardCenter, EInsetPolygonsMode::All, bShouldBevel, /* Out */ OutNewCenterPolygonRefs, /* Out */ OutNewSidePolygonRefs );
+}
+
+
+void UEditableMesh::GenerateNormalsAndTangentsForPolygons( const TArray< FPolygonRef >& PolygonRefs )
+{
+	static TArray<FVertexAttributesForPolygon> VertexAttributesForPolygons;
+	VertexAttributesForPolygons.Reset();
+	VertexAttributesForPolygons.Reserve( PolygonRefs.Num() );
+
+	for( const FPolygonRef& PolygonRef : PolygonRefs )
+	{
+		FVertexAttributesForPolygon& PolygonNewAttributes = *new( VertexAttributesForPolygons ) FVertexAttributesForPolygon();
+		PolygonNewAttributes.PolygonRef = PolygonRef;
+
+		const int32 PolygonPerimeterVertexCount = GetPolygonPerimeterVertexCount( PolygonRef );
+
+		PolygonNewAttributes.PerimeterVertexAttributeLists.SetNum( PolygonPerimeterVertexCount, false );
+
+		for( int32 VertexNumber = 0; VertexNumber < PolygonPerimeterVertexCount; ++VertexNumber )
+		{
+			// Compute the vertex normal
+			const FVector VertexNormal = ComputePolygonPerimeterVertexNormal( PolygonRef, VertexNumber );
+
+			// Save it!
+			// @todo mesheditor urgent perf: This causes every polygon vertex to be discreet, even if the vertex shares the same normal/uvs/tangents between coincident rendering vertices (common for smooth meshes)
+			TArray<FMeshElementAttributeData>& PerimeterVertexNewAttributes = PolygonNewAttributes.PerimeterVertexAttributeLists[ VertexNumber ].Attributes;
+			FMeshElementAttributeData& PerimeterVertexAttribute = *new( PerimeterVertexNewAttributes ) FMeshElementAttributeData( UEditableMeshAttribute::VertexNormal(), 0, FVector4( VertexNormal, 0.0f ) );
+		}
+	}
+	SetPolygonsVertexAttributes( VertexAttributesForPolygons );
+
+	// Now update the tangents
+	GenerateTangentsForPolygons( PolygonRefs );
+}
+
+
+void UEditableMesh::GenerateNormalsAndTangentsForPolygonsAndAdjacents( const TArray< FPolygonRef >& PolygonRefs )
+{
+	// Figure out everything next to these polygons
+	static TArray<FPolygonRef> PolygonRefsWithAdjacents;
+	PolygonRefsWithAdjacents.Reset();
+
+	PolygonRefsWithAdjacents = PolygonRefs;
+
+	for( const FPolygonRef PolygonRef : PolygonRefs )
+	{
+		static TArray<FVertexID> PerimeterVertexIDs;
+		PerimeterVertexIDs.Reset();
+		GetPolygonPerimeterVertices( PolygonRef, /* Out */ PerimeterVertexIDs );
+
+		for( const FVertexID PerimeterVertexID : PerimeterVertexIDs )
+		{
+			static TArray<FPolygonRef> ConnectedPolygonRefs;
+			GetVertexConnectedPolygons( PerimeterVertexID, /* Out */ ConnectedPolygonRefs );
+			for( const FPolygonRef ConnectedPolygonRef : ConnectedPolygonRefs )
+			{
+				// Ignore self
+				if( ConnectedPolygonRef != PolygonRef )
+				{
+					// Add uniquely, because it may have been one of the polygons that were passed in, or even a shared 
+					// neighbor of one of those polygons
+					PolygonRefsWithAdjacents.AddUnique( ConnectedPolygonRef );
+				}
+			}
+		}
+	}
+
+	return GenerateNormalsAndTangentsForPolygons( PolygonRefsWithAdjacents );
+}
+
+
+void UEditableMesh::GenerateTangentsForPolygons( const TArray< FPolygonRef >& PolygonRefs )
+{
+	static TArray<FVertexAttributesForPolygon> VertexAttributesForPolygons;
+	VertexAttributesForPolygons.Reset();
+	VertexAttributesForPolygons.Reserve( PolygonRefs.Num() );
+	for( const FPolygonRef& PolygonRef : PolygonRefs )
+	{
+		FVertexAttributesForPolygon& PolygonNewAttributes = *new( VertexAttributesForPolygons ) FVertexAttributesForPolygon();
+		PolygonNewAttributes.PolygonRef = PolygonRef;
+
+		const int32 PolygonPerimeterVertexCount = GetPolygonPerimeterVertexCount( PolygonRef );
+		PolygonNewAttributes.PerimeterVertexAttributeLists.SetNum( PolygonPerimeterVertexCount, false );
+	}
+
+	struct FMikkUserData
+	{
+		UEditableMesh* Self;
+		const TArray< FPolygonRef >& Polygons;
+		TArray<FVertexAttributesForPolygon>& VertexAttributesForPolygons;
+
+		FMikkUserData( UEditableMesh* InitSelf, const TArray< FPolygonRef >& InitPolygons, TArray<FVertexAttributesForPolygon>& InitVertexAttributesForPolygons )
+			: Self( InitSelf ),
+			  Polygons( InitPolygons ),
+			  VertexAttributesForPolygons( InitVertexAttributesForPolygons )
+		{
+		}
+	} MikkUserData( this, PolygonRefs, VertexAttributesForPolygons );
+
+	struct Local
+	{
+		static int MikkGetNumFaces( const SMikkTSpaceContext* Context )
+		{
+			const FMikkUserData& UserData = *static_cast<FMikkUserData*>( Context->m_pUserData );
+			return UserData.Polygons.Num();
+		}
+
+		static int MikkGetNumVertsOfFace( const SMikkTSpaceContext* Context, const int MikkFaceIndex )
+		{
+			const FMikkUserData& UserData = *static_cast<FMikkUserData*>( Context->m_pUserData );
+
+			const FPolygonRef& PolygonRef = UserData.Polygons[ MikkFaceIndex ];
+
+			// @todo mesheditor holes: Can we even use Mikktpsace for polygons with holes?  Do we need to run the triangulated triangles through instead?
+			const int32 PolygonVertexCount = UserData.Self->GetPolygonPerimeterVertexCount( PolygonRef );
+			return PolygonVertexCount;
+		}
+
+		static void MikkGetPosition( const SMikkTSpaceContext* Context, float OutPosition[3], const int MikkFaceIndex, const int MikkVertexIndex )
+		{
+			const FMikkUserData& UserData = *static_cast<FMikkUserData*>( Context->m_pUserData );
+
+			const FPolygonRef& PolygonRef = UserData.Polygons[ MikkFaceIndex ];
+			const FVector VertexPosition = UserData.Self->GetPolygonPerimeterVertexAttribute( PolygonRef, MikkVertexIndex, UEditableMeshAttribute::VertexPosition(), 0 );
+
+			OutPosition[0] = VertexPosition.X;
+			OutPosition[1] = VertexPosition.Y;
+			OutPosition[2] = VertexPosition.Z;
+		}
+
+		static void MikkGetNormal( const SMikkTSpaceContext* Context, float OutNormal[3], const int MikkFaceIndex, const int MikkVertexIndex )
+		{
+			const FMikkUserData& UserData = *static_cast<FMikkUserData*>( Context->m_pUserData );
+
+			const FPolygonRef& PolygonRef = UserData.Polygons[ MikkFaceIndex ];
+			const FVector PolygonVertexNormal = UserData.Self->GetPolygonPerimeterVertexAttribute( PolygonRef, MikkVertexIndex, UEditableMeshAttribute::VertexNormal(), 0 );
+
+			OutNormal[ 0 ] = PolygonVertexNormal.X;
+			OutNormal[ 1 ] = PolygonVertexNormal.Y;
+			OutNormal[ 2 ] = PolygonVertexNormal.Z;
+		}
+
+		static void MikkGetTexCoord( const SMikkTSpaceContext* Context, float OutUV[2], const int MikkFaceIndex, const int MikkVertexIndex )
+		{
+			const FMikkUserData& UserData = *static_cast<FMikkUserData*>( Context->m_pUserData );
+
+			// @todo mesheditor: Support using a custom texture coordinate index for tangent space generation?
+			const int32 TextureCoordinateIndex = 0;
+
+			const FPolygonRef& PolygonRef = UserData.Polygons[ MikkFaceIndex ];
+			const FVector2D PolygonVertexTextureCoordinate( UserData.Self->GetPolygonPerimeterVertexAttribute( PolygonRef, MikkVertexIndex, UEditableMeshAttribute::VertexTextureCoordinate(), TextureCoordinateIndex ) );
+
+			OutUV[0] = PolygonVertexTextureCoordinate.X;
+			OutUV[1] = PolygonVertexTextureCoordinate.Y;
+		}
+
+		static void MikkSetTSpaceBasic( const SMikkTSpaceContext* Context, const float Tangent[3], const float BitangentSign, const int MikkFaceIndex, const int MikkVertexIndex )
+		{
+			FMikkUserData& UserData = *static_cast<FMikkUserData*>( Context->m_pUserData );
+
+			const FVector NewTangent( Tangent[ 0 ], Tangent[ 1 ], Tangent[ 2 ] );
+
+			// Save it!
+			TArray<FMeshElementAttributeData>& PerimeterVertexNewAttributes = UserData.VertexAttributesForPolygons[ MikkFaceIndex ].PerimeterVertexAttributeLists[ MikkVertexIndex ].Attributes;
+			check( PerimeterVertexNewAttributes.Num() == 0 );
+			PerimeterVertexNewAttributes.Emplace( UEditableMeshAttribute::VertexTangent(), 0, FVector4( NewTangent, 0.0f ) );
+			PerimeterVertexNewAttributes.Emplace( UEditableMeshAttribute::VertexBinormalSign(), 0, FVector4( BitangentSign ) );
+		}
+	};
+
+	SMikkTSpaceInterface MikkTInterface;
+	{
+		MikkTInterface.m_getNumFaces = Local::MikkGetNumFaces;
+		MikkTInterface.m_getNumVerticesOfFace = Local::MikkGetNumVertsOfFace;
+		MikkTInterface.m_getPosition = Local::MikkGetPosition;
+		MikkTInterface.m_getNormal = Local::MikkGetNormal;
+		MikkTInterface.m_getTexCoord = Local::MikkGetTexCoord;
+
+		MikkTInterface.m_setTSpaceBasic = Local::MikkSetTSpaceBasic;
+		MikkTInterface.m_setTSpace = nullptr;
+	}
+
+	SMikkTSpaceContext MikkTContext;
+	{
+		MikkTContext.m_pInterface = &MikkTInterface;
+		MikkTContext.m_pUserData = (void*)( &MikkUserData );
+
+		// @todo mesheditor perf: Turning this on can apparently improve performance.  Needs investigation.
+		MikkTContext.m_bIgnoreDegenerates = false;
+	}
+
+	// Now we'll ask MikkTSpace to actually generate the tangents
+	genTangSpaceDefault( &MikkTContext );
+
+	SetPolygonsVertexAttributes( VertexAttributesForPolygons );
+}
+
+
+void UEditableMesh::SetVerticesCornerSharpness( const TArray<FVertexID>& VertexIDs, const TArray<float>& VerticesNewSharpness )
+{
+	check( VertexIDs.Num() == VerticesNewSharpness.Num() );
+
+	static TArray<FAttributesForVertex> AttributesForVertices;
+	AttributesForVertices.Reset();
+
+	for( int32 VertexNumber = 0; VertexNumber < VertexIDs.Num(); ++VertexNumber )
+	{
+		const FVertexID VertexID = VertexIDs[ VertexNumber ];
+
+		FAttributesForVertex& AttributesForVertex = *new( AttributesForVertices ) FAttributesForVertex();
+		AttributesForVertex.VertexID = VertexID;
+		AttributesForVertex.VertexAttributes.Attributes.Add( FMeshElementAttributeData(
+			UEditableMeshAttribute::VertexCornerSharpness(),
+			0,
+			FVector4( VerticesNewSharpness[ VertexNumber ] ) ) );
+	}
+
+	SetVerticesAttributes( AttributesForVertices );
+}
+
+
+void UEditableMesh::SetEdgesCreaseSharpness( const TArray<FEdgeID>& EdgeIDs, const TArray<float>& EdgesNewCreaseSharpness )
+{
+	check( EdgeIDs.Num() == EdgesNewCreaseSharpness.Num() );
+
+	static TArray<FAttributesForEdge> AttributesForEdges;
+	AttributesForEdges.Reset();
+
+	for( int32 EdgeNumber = 0; EdgeNumber < EdgeIDs.Num(); ++EdgeNumber )
+	{
+		const FEdgeID EdgeID = EdgeIDs[ EdgeNumber ];
+
+		FAttributesForEdge& AttributesForEdge = *new( AttributesForEdges ) FAttributesForEdge();
+		AttributesForEdge.EdgeID = EdgeID;
+		AttributesForEdge.EdgeAttributes.Attributes.Add( FMeshElementAttributeData(
+			UEditableMeshAttribute::EdgeCreaseSharpness(),
+			0,
+			FVector4( EdgesNewCreaseSharpness[ EdgeNumber ] ) ) );
+	}
+
+	SetEdgesAttributes( AttributesForEdges );
+}
+
+
+void UEditableMesh::SetEdgesHardness( const TArray<FEdgeID>& EdgeIDs, const TArray<bool>& EdgesNewIsHard )
+{
+	check( EdgeIDs.Num() == EdgesNewIsHard.Num() );
+
+	static TArray<FAttributesForEdge> AttributesForEdges;
+	AttributesForEdges.Reset();
+
+	static TSet<FPolygonRef> UniqueConnectedPolygonRefs;
+	UniqueConnectedPolygonRefs.Reset();
+
+	for( int32 EdgeNumber = 0; EdgeNumber < EdgeIDs.Num(); ++EdgeNumber )
+	{
+		const FEdgeID EdgeID = EdgeIDs[ EdgeNumber ];
+
+		FAttributesForEdge& AttributesForEdge = *new( AttributesForEdges ) FAttributesForEdge();
+		AttributesForEdge.EdgeID = EdgeID;
+		AttributesForEdge.EdgeAttributes.Attributes.Add( FMeshElementAttributeData(
+			UEditableMeshAttribute::EdgeIsHard(),
+			0,
+			FVector4( EdgesNewIsHard[ EdgeNumber ] ? 1.0f : 0.0f ) ) );
+
+		// Get the polygons this edge is connected to.  They'll need new normals.
+		static TArray<FPolygonRef> ConnectedPolygonRefs;
+		GetEdgeConnectedPolygons( EdgeID, /* Out */ ConnectedPolygonRefs );
+
+		for( const FPolygonRef ConnectedPolygonRef : ConnectedPolygonRefs )
+		{
+			UniqueConnectedPolygonRefs.Add( ConnectedPolygonRef );
+		}
+	}
+
+	SetEdgesAttributes( AttributesForEdges );
+
+	// Generate fresh normals
+	GenerateNormalsAndTangentsForPolygons( MoveTemp( UniqueConnectedPolygonRefs.Array() ) );
+}
+
+
+void UEditableMesh::SetEdgesHardnessAutomatically( const TArray<FEdgeID>& EdgeIDs, const float MaxDotProductForSoftEdge )
+{
+	static TArray<bool> EdgesNewIsHard;
+	EdgesNewIsHard.SetNumUninitialized( EdgeIDs.Num(), false );
+
+	for( int32 EdgeNumber = 0; EdgeNumber < EdgeIDs.Num(); ++EdgeNumber )
+	{
+		const FEdgeID EdgeID = EdgeIDs[ EdgeNumber ];
+
+		// Default to soft if we have no polygons attached
+		bool bIsSoftEdge = true;
+
+		const int32 ConnectedPolygonCount = GetEdgeConnectedPolygonCount( EdgeID );
+		if( ConnectedPolygonCount > 0 )
+		{
+			float MinDot = 1.0f;
+
+			const FPolygonRef FirstPolygonRef = GetEdgeConnectedPolygon( EdgeID, 0 );
+
+			FVector LastPolygonNormal = ComputePolygonNormal( FirstPolygonRef );
+
+			for( int32 ConnectedPolygonNumber = 1; ConnectedPolygonNumber < ConnectedPolygonCount; ++ConnectedPolygonNumber )
+			{
+				const FPolygonRef PolygonRef = GetEdgeConnectedPolygon( EdgeID, ConnectedPolygonNumber );
+
+				const FVector PolygonNormal = ComputePolygonNormal( PolygonRef );
+
+				const float Dot = FVector::DotProduct( PolygonNormal, LastPolygonNormal );
+				MinDot = FMath::Min( Dot, MinDot );
+			}
+
+			bIsSoftEdge = ( MinDot >= MaxDotProductForSoftEdge );
+		}
+
+		EdgesNewIsHard[ EdgeNumber ] = !bIsSoftEdge;
+	}
+
+
+	// Set the edges hardness (and generate new normals)
+	SetEdgesHardness( EdgeIDs, EdgesNewIsHard );
+}
+
+
+void UEditableMesh::SetEdgesVertices( const TArray<FVerticesForEdge>& VerticesForEdges )
+{
+	FSetEdgesVerticesChangeInput RevertInput;
+
+	RevertInput.VerticesForEdges.AddUninitialized( VerticesForEdges.Num() );
+	for( int32 VertexNumber = 0; VertexNumber < VerticesForEdges.Num(); ++VertexNumber )
+	{
+		const FVerticesForEdge& VerticesForEdge = VerticesForEdges[ VertexNumber ];
+
+		// Save the backup
+		FVerticesForEdge& RevertVerticesForEdge = RevertInput.VerticesForEdges[ VertexNumber ];
+		RevertVerticesForEdge.EdgeID = VerticesForEdge.EdgeID;
+		GetEdgeVertices( VerticesForEdge.EdgeID, /* Out */ RevertVerticesForEdge.NewVertexID0, /* Out */ RevertVerticesForEdge.NewVertexID1 );
+
+		// Edit the edge
+		SetEdgeVertices_Internal( VerticesForEdge.EdgeID, VerticesForEdge.NewVertexID0, VerticesForEdge.NewVertexID1 );
+	}
+
+	AddUndo( MakeUnique<FSetEdgesVerticesChange>( MoveTemp( RevertInput ) ) );
+}
+
+
+void UEditableMesh::InsertPolygonPerimeterVertices( const FPolygonRef PolygonRef, const int32 InsertBeforeVertexNumber, const TArray<FVertexAndAttributes>& VerticesToInsert )
+{
+	FRemovePolygonPerimeterVerticesChangeInput RevertInput;
+
+	RevertInput.PolygonRef = PolygonRef;
+	RevertInput.FirstVertexNumberToRemove = InsertBeforeVertexNumber;
+	RevertInput.NumVerticesToRemove = VerticesToInsert.Num();
+
+	InsertPolygonPerimeterVertices_Internal( PolygonRef, InsertBeforeVertexNumber, VerticesToInsert );
+
+	AddUndo( MakeUnique<FRemovePolygonPerimeterVerticesChange>( MoveTemp( RevertInput ) ) );
+}
+
+
+void UEditableMesh::RemovePolygonPerimeterVertices( const FPolygonRef PolygonRef, const int32 FirstVertexNumberToRemove, const int32 NumVerticesToRemove )
+{
+	FInsertPolygonPerimeterVerticesChangeInput RevertInput;
+
+	RevertInput.PolygonRef = PolygonRef;
+	RevertInput.InsertBeforeVertexNumber = FirstVertexNumberToRemove;
+
+	RevertInput.VerticesToInsert.SetNum( NumVerticesToRemove, false );
+	for( int32 VertexToRemoveIter = 0; VertexToRemoveIter < NumVerticesToRemove; ++VertexToRemoveIter )
+	{
+		FVertexAndAttributes& RevertVertexToInsert = RevertInput.VerticesToInsert[ VertexToRemoveIter ];
+		RevertVertexToInsert.VertexID = GetPolygonPerimeterVertex( PolygonRef, FirstVertexNumberToRemove + VertexToRemoveIter );
+
+		for( const FName AttributeName : UEditableMesh::GetValidPolygonVertexAttributes() )
+		{
+			const int32 MaxAttributeIndex = GetMaxAttributeIndex( AttributeName );
+			for( int32 AttributeIndex = 0; AttributeIndex < MaxAttributeIndex; ++AttributeIndex )
+			{
+				FMeshElementAttributeData& PolygonVertexAttribute = *new( RevertVertexToInsert.PolygonVertexAttributes.Attributes ) FMeshElementAttributeData(
+					AttributeName,
+					AttributeIndex,
+					GetPolygonPerimeterVertexAttribute( PolygonRef, FirstVertexNumberToRemove + VertexToRemoveIter, AttributeName, AttributeIndex ) );
+			}
+		}
+	}
+
+	RemovePolygonPerimeterVertices_Internal( PolygonRef, FirstVertexNumberToRemove, NumVerticesToRemove );
+
+	AddUndo( MakeUnique<FInsertPolygonPerimeterVerticesChange>( MoveTemp( RevertInput ) ) );
+}
+
+
+int32 UEditableMesh::FindPolygonPerimeterVertexNumberForVertex( const FPolygonRef PolygonRef, const FVertexID VertexID ) const
+{
+	int32 FoundPolygonVertexNumber = INDEX_NONE;
+
+	const int32 PolygonPerimeterVertexCount = this->GetPolygonPerimeterVertexCount( PolygonRef );
+	for( int32 PolygonVertexNumber = 0; PolygonVertexNumber < PolygonPerimeterVertexCount; ++PolygonVertexNumber )
+	{
+		if( VertexID == this->GetPolygonPerimeterVertex( PolygonRef, PolygonVertexNumber ) )
+		{
+			FoundPolygonVertexNumber = PolygonVertexNumber;
+			break;
+		}
+	}
+
+	return FoundPolygonVertexNumber;
+}
+
+
+int32 UEditableMesh::FindPolygonHoleVertexNumberForVertex( const FPolygonRef PolygonRef, const int32 HoleNumber, const FVertexID VertexID ) const
+{
+	int32 FoundPolygonVertexNumber = INDEX_NONE;
+
+	const int32 PolygonHoleVertexCount = this->GetPolygonHoleVertexCount( PolygonRef, HoleNumber );
+	for( int32 PolygonVertexNumber = 0; PolygonVertexNumber < PolygonHoleVertexCount; ++PolygonVertexNumber )
+	{
+		if( VertexID == this->GetPolygonHoleVertex( PolygonRef, HoleNumber, PolygonVertexNumber ) )
+		{
+			FoundPolygonVertexNumber = PolygonVertexNumber;
+			break;
+		}
+	}
+
+	return FoundPolygonVertexNumber;
+}
+
+
+int32 UEditableMesh::FindPolygonPerimeterEdgeNumberForVertices( const FPolygonRef PolygonRef, const FVertexID EdgeVertexID0, const FVertexID EdgeVertexID1 ) const
+{
+	int32 FoundPolygonEdgeNumber = INDEX_NONE;
+
+	static TArray<FEdgeID> EdgeIDs;
+	GetPolygonPerimeterEdges( PolygonRef, /* Out */ EdgeIDs );
+
+	for( int32 PolygonEdgeNumber = 0; PolygonEdgeNumber < EdgeIDs.Num(); ++PolygonEdgeNumber )
+	{
+		const FEdgeID EdgeID = EdgeIDs[ PolygonEdgeNumber ];
+
+		FVertexID TestEdgeVertexIDs[ 2 ];
+		GetEdgeVertices( EdgeID, /* Out */ TestEdgeVertexIDs[0], /* Out */ TestEdgeVertexIDs[1] );
+
+		if( ( TestEdgeVertexIDs[ 0 ] == EdgeVertexID0 && TestEdgeVertexIDs[ 1 ] == EdgeVertexID1 ) ||
+			( TestEdgeVertexIDs[ 1 ] == EdgeVertexID0 && TestEdgeVertexIDs[ 0 ] == EdgeVertexID1 ) )
+		{
+			FoundPolygonEdgeNumber = PolygonEdgeNumber;
+			break;
+		}
+	}
+
+	return FoundPolygonEdgeNumber;
+}
+
+
+int32 UEditableMesh::FindPolygonHoleEdgeNumberForVertices( const FPolygonRef PolygonRef, const int32 HoleNumber, const FVertexID EdgeVertexID0, const FVertexID EdgeVertexID1 ) const
+{
+	int32 FoundPolygonEdgeNumber = INDEX_NONE;
+
+	static TArray<FEdgeID> EdgeIDs;
+	GetPolygonHoleEdges( PolygonRef, HoleNumber, /* Out */ EdgeIDs );
+
+	for( int32 PolygonEdgeNumber = 0; PolygonEdgeNumber < EdgeIDs.Num(); ++PolygonEdgeNumber )
+	{
+		const FEdgeID EdgeID = EdgeIDs[ PolygonEdgeNumber ];
+
+		FVertexID TestEdgeVertexIDs[ 2 ];
+		GetEdgeVertices( EdgeID, /* Out */ TestEdgeVertexIDs[0], /* Out */ TestEdgeVertexIDs[1] );
+
+		if( ( TestEdgeVertexIDs[ 0 ] == EdgeVertexID0 && TestEdgeVertexIDs[ 1 ] == EdgeVertexID1 ) ||
+			( TestEdgeVertexIDs[ 1 ] == EdgeVertexID0 && TestEdgeVertexIDs[ 0 ] == EdgeVertexID1 ) )
+		{
+			FoundPolygonEdgeNumber = PolygonEdgeNumber;
+			break;
+		}
+	}
+
+	return FoundPolygonEdgeNumber;
+}
+
+
+void UEditableMesh::FlipPolygons( const TArray<FPolygonRef>& PolygonRefs )
+{
+	static TArray<FPolygonToCreate> PolygonsToCreate;
+	PolygonsToCreate.Reset();
+	for( int32 PolygonNumber = 0; PolygonNumber < PolygonRefs.Num(); ++PolygonNumber )
+	{
+		const FPolygonRef OriginalPolygonRef = PolygonRefs[ PolygonNumber ];
+
+		FPolygonToCreate& PolygonToCreate = *new( PolygonsToCreate ) FPolygonToCreate();
+
+		PolygonToCreate.SectionID = OriginalPolygonRef.SectionID;
+		
+		// Keep the original polygon ID.  No reason not to.
+		PolygonToCreate.OriginalPolygonID = OriginalPolygonRef.PolygonID;
+
+		// Iterate backwards to add the vertices for the polygon, because we're flipping it over.
+		const int32 PerimeterVertexCount = GetPolygonPerimeterVertexCount( OriginalPolygonRef );
+		PolygonToCreate.PerimeterVertices.Reserve( PerimeterVertexCount );
+		for( int32 VertexNumber = PerimeterVertexCount - 1; VertexNumber >= 0; --VertexNumber )
+		{
+			FVertexAndAttributes& PerimeterVertex = *new( PolygonToCreate.PerimeterVertices ) FVertexAndAttributes();
+
+			PerimeterVertex.VertexID = GetPolygonPerimeterVertex( OriginalPolygonRef, VertexNumber );
+
+			// Save off information about the polygons, so we can re-create them with reverse winding.
+			TArray<FMeshElementAttributeData>& PerimeterVertexAttributeList = PerimeterVertex.PolygonVertexAttributes.Attributes;
+			for( const FName AttributeName : UEditableMesh::GetValidPolygonVertexAttributes() )
+			{
+				const int32 MaxAttributeIndex = GetMaxAttributeIndex( AttributeName );
+				for( int32 AttributeIndex = 0; AttributeIndex < MaxAttributeIndex; ++AttributeIndex )
+				{
+					FMeshElementAttributeData& PolygonVertexAttribute = *new( PerimeterVertexAttributeList ) FMeshElementAttributeData(
+						AttributeName,
+						AttributeIndex,
+						GetPolygonPerimeterVertexAttribute( OriginalPolygonRef, VertexNumber, AttributeName, AttributeIndex ) );
+				}
+			}
+		}
+
+		// @todo mesheditor hole: Perserve holes in this polygon!  (Vertex IDs and Attributes!) Should these be flipped too?  Probably.
+	}
+
+	// Delete all of the polygons
+	{
+		const bool bDeleteOrphanedEdges = false;
+		const bool bDeleteOrphanedVertices = false;
+		const bool bDeleteEmptySections = false;
+		DeletePolygons( PolygonRefs, bDeleteOrphanedEdges, bDeleteOrphanedVertices, bDeleteEmptySections );
+	}
+
+	// Create the new polygons.  They're just like the polygons we deleted, except with reversed winding.
+	{
+		static TArray<FPolygonRef> NewPolygonRefs;
+		static TArray<FEdgeID> NewEdgeIDs;
+		CreatePolygons( PolygonsToCreate, NewPolygonRefs, NewEdgeIDs );
+
+		// Generate fresh normals
+		GenerateNormalsAndTangentsForPolygonsAndAdjacents( NewPolygonRefs );
+	}
+}
+
+
+void UEditableMesh::TriangulatePolygons( const TArray<FPolygonRef>& PolygonRefs, TArray<FPolygonRef>& OutNewTrianglePolygons )
+{
+	OutNewTrianglePolygons.Reset();
+
+	static TArray<FPolygonToCreate> PolygonsToCreate;
+	PolygonsToCreate.Reset();
+
+	static TArray<FPolygonRef> PolygonsToDelete;
+	PolygonsToDelete.Reset();
+
+	for( const FPolygonRef PolygonRef : PolygonRefs )
+	{
+		// Skip right over polygons with fewer than four vertices
+		const int32 PerimeterVertexCount = GetPolygonPerimeterVertexCount( PolygonRef );
+		if( PerimeterVertexCount > 3 )
+		{
+			// We'll be replacing this polygon with it's triangulated counterpart polygons
+			PolygonsToDelete.Add( PolygonRef );
+
+			// Figure out the triangulation for this polygon
+			static TArray<int32> PerimeterVertexNumbersForTriangles;
+			PerimeterVertexNumbersForTriangles.Reset();
+			ComputePolygonTriangulation( PolygonRef, /* Out */ PerimeterVertexNumbersForTriangles );
+
+			// Build polygons for each of the triangles that made up the original
+			{
+				check( ( PerimeterVertexNumbersForTriangles.Num() % 3 ) == 0 );	// Expecting three indices per triangle
+				const int32 TriangleCount = PerimeterVertexNumbersForTriangles.Num() / 3;
+				for( int32 TriangleIndex = 0; TriangleIndex < TriangleCount; ++TriangleIndex )
+				{
+					const int32 FirstTrianglesVertexIDIndex = TriangleIndex * 3;
+
+					FPolygonToCreate& PolygonToCreate = *new( PolygonsToCreate ) FPolygonToCreate();
+					PolygonToCreate.SectionID = PolygonRef.SectionID;
+
+					for( int32 TriangleVertexNumber = 0; TriangleVertexNumber < 3; ++TriangleVertexNumber )
+					{
+						const int32 PerimeterVertexNumber = PerimeterVertexNumbersForTriangles[ FirstTrianglesVertexIDIndex + TriangleVertexNumber ];
+						const FVertexID TriangleVertexID = GetPolygonPerimeterVertex( PolygonRef, PerimeterVertexNumber );
+
+						FVertexAndAttributes& PerimeterVertex = *new( PolygonToCreate.PerimeterVertices ) FVertexAndAttributes();
+						PerimeterVertex.VertexID = TriangleVertexID;
+
+						// Save off information about the polygons, so we can re-create them with reverse winding.
+						TArray<FMeshElementAttributeData>& PerimeterVertexAttributeList = PerimeterVertex.PolygonVertexAttributes.Attributes;
+						for( const FName AttributeName : UEditableMesh::GetValidPolygonVertexAttributes() )
+						{
+							const int32 MaxAttributeIndex = GetMaxAttributeIndex( AttributeName );
+							for( int32 AttributeIndex = 0; AttributeIndex < MaxAttributeIndex; ++AttributeIndex )
+							{
+								FMeshElementAttributeData& PolygonVertexAttribute = *new( PerimeterVertexAttributeList ) FMeshElementAttributeData(
+									AttributeName,
+									AttributeIndex,
+									GetPolygonPerimeterVertexAttribute( PolygonRef, PerimeterVertexNumber, AttributeName, AttributeIndex ) );
+							}
+						}
+					}
+				}
+			}
+		}
+	}
+
+	// Delete the original polygons, but don't erase any orphaned edges or vertices, because we're about to put in
+	// triangles to replace those polygons.  Also, we won't touch polygons that we didn't have to triangulate!
+	{
+		const bool bDeleteOrphanEdges = false;
+		const bool bDeleteOrphanVertices = false;
+		const bool bDeleteEmptySections = false;
+		DeletePolygons( PolygonsToDelete, bDeleteOrphanEdges, bDeleteOrphanVertices, bDeleteEmptySections );
+	}
+
+	// Create the new polygons.  One for each triangle.  Note that new edges will be created here too on the inside of
+	// the original polygon to border the triangles.
+	{
+		static TArray<FEdgeID> NewEdgeIDs;
+		CreatePolygons( PolygonsToCreate, /* Out */ OutNewTrianglePolygons, /* Out */ NewEdgeIDs );
+
+		// @todo mesheditor: The the new internal edges of the polygon (NewEdgeIDs) get hard or soft edges by default?\
+
+		// Generate fresh normals
+		GenerateNormalsAndTangentsForPolygonsAndAdjacents( OutNewTrianglePolygons );
+	}
+}
+
+
+FSectionID UEditableMesh::CreateSection( const FSectionToCreate& SectionToCreate )
+{
+	return CreateSection_Internal( SectionToCreate );
+}
+
+
+void UEditableMesh::DeleteSection( const FSectionID SectionID )
+{
+	DeleteSection_Internal( SectionID );
+}
+
+
+void UEditableMesh::AssignMaterialToPolygons( const TArray<FPolygonRef>& PolygonRefs, UMaterialInterface* Material, TArray<FPolygonRef>& NewPolygonRefs )
+{
+	// We need to move polygons from one section to another. First find the new section id to move them to.
+	const bool bCreateNewSectionIfNotFound = true;
+	const FSectionID NewSectionID = GetSectionIDFromMaterial_Internal( Material, bCreateNewSectionIfNotFound );
+	check( NewSectionID != FSectionID::Invalid );
+
+	// Make an array of FPolygonToCreate, based on the ones we wish to move. Everything is the same, other than the SectionID
+	static TArray<FPolygonToCreate> PolygonsToCreate;
+	PolygonsToCreate.Reset( PolygonRefs.Num() );
+	for( int32 PolygonNumber = 0; PolygonNumber < PolygonRefs.Num(); ++PolygonNumber )
+	{
+		const FPolygonRef OriginalPolygonRef = PolygonRefs[ PolygonNumber ];
+
+		FPolygonToCreate& PolygonToCreate = *new( PolygonsToCreate ) FPolygonToCreate();
+
+		PolygonToCreate.SectionID = NewSectionID;
+		PolygonToCreate.OriginalPolygonID = FPolygonID::Invalid;
+
+		const int32 PerimeterVertexCount = GetPolygonPerimeterVertexCount( OriginalPolygonRef );
+		PolygonToCreate.PerimeterVertices.Reserve( PerimeterVertexCount );
+		for( int32 VertexNumber = 0; VertexNumber < PerimeterVertexCount; ++VertexNumber )
+		{
+			FVertexAndAttributes& PerimeterVertex = *new( PolygonToCreate.PerimeterVertices ) FVertexAndAttributes();
+
+			PerimeterVertex.VertexID = GetPolygonPerimeterVertex( OriginalPolygonRef, VertexNumber );
+
+			TArray<FMeshElementAttributeData>& PerimeterVertexAttributeList = PerimeterVertex.PolygonVertexAttributes.Attributes;
+			for( const FName AttributeName : UEditableMesh::GetValidPolygonVertexAttributes() )
+			{
+				const int32 MaxAttributeIndex = GetMaxAttributeIndex( AttributeName );
+				for( int32 AttributeIndex = 0; AttributeIndex < MaxAttributeIndex; ++AttributeIndex )
+				{
+					FMeshElementAttributeData& PolygonVertexAttribute = *new( PerimeterVertexAttributeList ) FMeshElementAttributeData(
+						AttributeName,
+						AttributeIndex,
+						GetPolygonPerimeterVertexAttribute( OriginalPolygonRef, VertexNumber, AttributeName, AttributeIndex ) );
+				}
+			}
+		}
+
+		// @todo mesheditor: Handle holes?
+	}
+
+	// Create the new polygons in the new section.
+	{
+		static TArray<FEdgeID> NewEdgeIDs;
+		CreatePolygons( PolygonsToCreate, NewPolygonRefs, NewEdgeIDs );
+	}
+
+	// Delete the polygons from the old section
+	{
+		const bool bDeleteOrphanedEdges = false;
+		const bool bDeleteOrphanedVertices = false;
+		const bool bDeleteEmptySections = true;
+		DeletePolygons( PolygonRefs, bDeleteOrphanedEdges, bDeleteOrphanedVertices, bDeleteEmptySections );
+	}
+}
+
+
+void UEditableMesh::WeldVertices( const TArray<FVertexID>& VertexIDsToWeld, FVertexID& OutNewVertexID )
+{
+	OutNewVertexID = FVertexID::Invalid;
+
+	// This function takes a list of perimeter vertices and a list of vertices to be welded as input.
+	// It returns a tuple stating whether the result is valid, and the [first, last) range of vertices to be welded.
+	// (It will be invalid if there is more than one contiguous run of vertices to weld.)
+	auto GetPerimeterVertexRangeToWeld = []( const TArray<FVertexID>& PolygonVertexIDs, const TArray<FVertexID>& VertexIDsToWeld )
+	{
+		bool bValid = true;
+		int32 StartIndex = INDEX_NONE;
+		int32 EndIndex = INDEX_NONE;
+
+		const int32 NumPolygonVertices = PolygonVertexIDs.Num();
+		bool bPrevVertexNeedsWelding = VertexIDsToWeld.Contains( PolygonVertexIDs[ NumPolygonVertices - 1 ] );
+		for( int32 Index = 0; Index < NumPolygonVertices; ++Index )
+		{
+			const bool bThisVertexNeedsWelding = VertexIDsToWeld.Contains( PolygonVertexIDs[ Index ] );
+			if( !bPrevVertexNeedsWelding && bThisVertexNeedsWelding )
+			{
+				// Transition from 'doesn't need welding' to 'needs welding'
+				if( StartIndex == INDEX_NONE )
+				{
+					StartIndex = Index;
+				}
+				else
+				{
+					// If this is not the first time we've seen this transition, there is more than one contiguous run of vertices
+					// which need welding, which is not allowed.
+					bValid = false;
+				}
+			}
+
+			if( bPrevVertexNeedsWelding && !bThisVertexNeedsWelding )
+			{
+				// Transition from 'needs welding' to 'doesn't need welding'
+				if( EndIndex == INDEX_NONE )
+				{
+					EndIndex = Index;
+				}
+				else
+				{
+					bValid = false;
+				}
+			}
+
+			bPrevVertexNeedsWelding = bThisVertexNeedsWelding;
+		}
+
+		// If the indices are not set, either there were no vertices to weld, or they were all to be welded.
+		// In the latter case, initialize the full vertex range.
+		if( StartIndex == INDEX_NONE && EndIndex == INDEX_NONE && bPrevVertexNeedsWelding )
+		{
+			StartIndex = 0;
+			EndIndex = NumPolygonVertices;
+		}
+
+		// Get the size of the range.
+		// The array is circular, so it's possible for the EndIndex to be smaller than the start index (and compensate for that accordingly)
+		const int32 RangeSize = ( EndIndex - StartIndex ) + ( ( EndIndex < StartIndex ) ? NumPolygonVertices : 0 );
+
+		// If, after welding perimeter vertices, we have fewer than three vertices left, this poly just disappears.
+		// (+ 1 below for the new vertex which replaces the welded range)
+		bool bWouldBeDegenerate = ( NumPolygonVertices - RangeSize + 1 < 3 );
+
+		return MakeTuple( bValid, bWouldBeDegenerate, StartIndex, EndIndex );
+	};
+
+	// Build a list of all polygons which contain at least one of the vertices to be welded
+	static TArray<FPolygonRef> AllConnectedPolygonRefs;
+	{
+		AllConnectedPolygonRefs.Reset();
+
+		for( FVertexID VertexID : VertexIDsToWeld )
+		{
+			static TArray<FPolygonRef> ConnectedPolygonRefs;
+			GetVertexConnectedPolygons( VertexID, ConnectedPolygonRefs );
+
+			for( FPolygonRef PolygonRef : ConnectedPolygonRefs )
+			{
+				AllConnectedPolygonRefs.AddUnique( PolygonRef );
+			}
+		}
+	}
+
+	// Check whether the operation is valid. We can't weld vertices if there are any polygons which have non-contiguous
+	// vertices on the perimeter which are marked to be welded
+	bool bNeedToCreateWeldedVertex = false;
+	for( FPolygonRef ConnectedPolygonRef : AllConnectedPolygonRefs )
+	{
+		static TArray<FVertexID> PolygonVertexIDs;
+		GetPolygonPerimeterVertices( ConnectedPolygonRef, PolygonVertexIDs );
+
+		auto VertexRangeToWeld = GetPerimeterVertexRangeToWeld( PolygonVertexIDs, VertexIDsToWeld );
+		const bool bIsValid = VertexRangeToWeld.Get<0>();
+		const bool bWouldBeDegenerate = VertexRangeToWeld.Get<1>();
+
+		// If the resulting poly is valid (has 3 or more verts), we know we need to create the welded vertex
+		if( !bWouldBeDegenerate )
+		{
+			bNeedToCreateWeldedVertex = true;
+		}
+
+		// If the result is invalid (because it would cause a poly to be welded in more than one place on its perimeter), abort now
+		if( !bIsValid )
+		{
+			// Return with the NewVertexID set to Invalid
+			return;
+		}
+	}
+
+	if( !bNeedToCreateWeldedVertex )
+	{
+		// For now, abort if we don't need to create a welded vertex.
+		// This generally implies that all (or a disconnected subset) of the mesh is about to disappear,
+		// which arguably is not something we would want to do like this anyway.
+		return;
+	}
+
+	// Create new welded vertex
+	static TArray<FVertexID> NewVertices;
+	{
+		static TArray<FVertexToCreate> VerticesToCreate;
+		VerticesToCreate.SetNum( 1 );
+		FVertexToCreate& VertexToCreate = VerticesToCreate.Last();
+
+		// The vertex which is created will be at the position of the last vertex in the array of vertices to weld.
+		// @todo mesheditor: maybe specify a position to the method instead of using one of the existing vertices?
+		VertexToCreate.VertexAttributes.Attributes.Emplace(
+			UEditableMeshAttribute::VertexPosition(),
+			0,
+			GetVertexAttribute( VertexIDsToWeld.Last(), UEditableMeshAttribute::VertexPosition(), 0 )
+		);
+
+		// @todo mesheditor: vertex corner sharpness too?
+
+		CreateVertices( VerticesToCreate, NewVertices );
+	}
+
+
+	static TArray<FPolygonToCreate> PolygonsToCreate;
+	PolygonsToCreate.Reset( AllConnectedPolygonRefs.Num() );
+
+	static TArray<FAttributesForEdge> AttributesForEdges;
+	AttributesForEdges.Reset();
+
+	// Now for each polygon, merge runs of consecutive vertices
+	for( FPolygonRef ConnectedPolygonRef : AllConnectedPolygonRefs )
+	{
+		const int32 NumPerimeterVertices = GetPolygonPerimeterVertexCount( ConnectedPolygonRef );
+
+		// Get perimeter vertices and edges for this polygon
+		static TArray<FVertexID> PolygonVertexIDs;
+		static TArray<FEdgeID> PolygonEdgeIDs;
+		GetPolygonPerimeterVertices( ConnectedPolygonRef, PolygonVertexIDs );
+		GetPolygonPerimeterEdges( ConnectedPolygonRef, PolygonEdgeIDs );
+
+		// Get the index range of perimeter vertices to be welded.
+		// This should definitely be valid, as any invalid welded poly will have caused early exit, above.
+		auto VertexRangeToWeld = GetPerimeterVertexRangeToWeld( PolygonVertexIDs, VertexIDsToWeld );
+		const bool bIsValid = VertexRangeToWeld.Get<0>();
+		const bool bWouldBeDegenerate = VertexRangeToWeld.Get<1>();
+		const int32 StartIndex = VertexRangeToWeld.Get<2>();
+		const int32 EndIndex = VertexRangeToWeld.Get<3>();
+		check( bIsValid );
+
+		if( bWouldBeDegenerate )
+		{
+			continue;
+		}
+
+		// Prepare to create a new polygon
+		PolygonsToCreate.Emplace();
+		FPolygonToCreate& PolygonToCreate = PolygonsToCreate.Last();
+		PolygonToCreate.SectionID = ConnectedPolygonRef.SectionID;
+
+		// Iterate through perimeter vertices starting at index 0.
+		// We skip through the run of welded vertices, replacing them with a single welded vertex.
+		// We need to check whether we are starting in the middle of a run (if EndIndex < StartIndex).
+		bool bInsideWeldedRange = ( EndIndex < StartIndex );
+		for( int32 Index = 0; Index < NumPerimeterVertices; ++Index )
+		{
+			if( bInsideWeldedRange )
+			{
+				if( Index == EndIndex )
+				{
+					// EndIndex is range exclusive, so we now need to process this vertex.
+					bInsideWeldedRange = false;
+				}
+				else
+				{
+					// Otherwise still inside the welded range; skip the remaining vertices in the range.
+					continue;
+				}
+			}
+
+			// Add new perimeter vertex in the polygon to create
+			PolygonToCreate.PerimeterVertices.Emplace();
+			FVertexAndAttributes& VertexAndAttributes = PolygonToCreate.PerimeterVertices.Last();
+
+			if( Index == StartIndex )
+			{
+				// If this is the first vertex in the run of vertices to weld, replace the ID with the newly created welded vertex
+				VertexAndAttributes.VertexID = NewVertices[ 0 ];
+				bInsideWeldedRange = true;
+			}
+			else
+			{
+				// Otherwise use the original Vertex ID
+				VertexAndAttributes.VertexID = PolygonVertexIDs[ Index ];
+			}
+
+			// Copy the polygon vertex attributes over
+			for( int32 TextureCoordinateIndex = 0; TextureCoordinateIndex < TextureCoordinateCount; ++TextureCoordinateIndex )
+			{
+				VertexAndAttributes.PolygonVertexAttributes.Attributes.Add( FMeshElementAttributeData(
+					UEditableMeshAttribute::VertexTextureCoordinate(),
+					TextureCoordinateIndex,
+					GetPolygonPerimeterVertexAttribute( ConnectedPolygonRef, Index, UEditableMeshAttribute::VertexTextureCoordinate(), TextureCoordinateIndex ) )
+				);
+			}
+
+			VertexAndAttributes.PolygonVertexAttributes.Attributes.Add( FMeshElementAttributeData(
+				UEditableMeshAttribute::VertexColor(),
+				0,
+				GetPolygonPerimeterVertexAttribute( ConnectedPolygonRef, Index, UEditableMeshAttribute::VertexColor(), 0 ) )
+			);
+
+			// Prepare to assign the old edge's attributes to the new edge.
+			// We build up an array of edge attributes to set, in perimeter vertex order for each polygon.
+			AttributesForEdges.Emplace();
+			FAttributesForEdge& AttributesForEdge = AttributesForEdges.Last();
+
+			AttributesForEdge.EdgeAttributes.Attributes.Add( FMeshElementAttributeData(
+				UEditableMeshAttribute::EdgeIsHard(),
+				0,
+				GetEdgeAttribute( PolygonEdgeIDs[ Index ], UEditableMeshAttribute::EdgeIsHard(), 0 ) )
+			);
+			AttributesForEdge.EdgeAttributes.Attributes.Add( FMeshElementAttributeData( 
+				UEditableMeshAttribute::EdgeCreaseSharpness(), 
+				0, 
+				GetEdgeAttribute( PolygonEdgeIDs[ Index ], UEditableMeshAttribute::EdgeCreaseSharpness(), 0 ) )
+			);
+		}
+	}
+
+	// Create polygons
+	static TArray<FPolygonRef> NewPolygonRefs;
+	{
+		static TArray<FEdgeID> NewEdgeIDs;
+		CreatePolygons( PolygonsToCreate, NewPolygonRefs, NewEdgeIDs );
+	}
+
+	// Set new edge attributes.
+	// Now that we have a list of newly created polygon refs, we need to go through the attributes for edge list, filling in the new Edge ID.
+	// This relies on the fact that the NewPolygonRefs array lists the polygons in the same order as they were defined in PolygonsToCreate,
+	// and that the edges are strictly ordered from perimeter vertex 0.
+	{
+		int32 AttributesForEdgeIndex = 0;
+		for( FPolygonRef NewPolygonRef : NewPolygonRefs )
+		{
+			static TArray<FEdgeID> NewPolygonEdgeIDs;
+			GetPolygonPerimeterEdges( NewPolygonRef, NewPolygonEdgeIDs );
+
+			for( FEdgeID NewPolygonEdgeID : NewPolygonEdgeIDs )
+			{
+				AttributesForEdges[ AttributesForEdgeIndex ].EdgeID = NewPolygonEdgeID;
+				AttributesForEdgeIndex++;
+			}
+		}
+		check( AttributesForEdgeIndex == AttributesForEdges.Num() );
+		SetEdgesAttributes( AttributesForEdges );
+	}
+
+	// Delete old polygons, removing any orphaned edges and vertices at the same time
+	{
+		const bool bDeleteOrphanedEdges = true;
+		const bool bDeleteOrphanedVertices = true;
+		const bool bDeleteEmptySections = false;
+		DeletePolygons( AllConnectedPolygonRefs, bDeleteOrphanedEdges, bDeleteOrphanedVertices, bDeleteEmptySections );
+	}
+
+	GenerateNormalsAndTangentsForPolygonsAndAdjacents( NewPolygonRefs );
+}
+
+
+void UEditableMesh::TessellatePolygons( const TArray<FPolygonRef>& PolygonRefs, const ETriangleTessellationMode TriangleTessellationMode, TArray<FPolygonRef>& OutNewPolygonRefs )
+{
+	OutNewPolygonRefs.Reset();
+
+	//
+	// Simple tessellation algorithm:
+	//
+	//   - Triangles will be split into either three or four triangles depending on the 'Mode' argument.
+	//			-> ThreeTriangles: Connect each vertex to a new center vertex, forming three triangles
+	//			-> FourTriangles: Split each edge and create a center polygon that connects those new vertices, then three additional polygons for each original corner
+	//
+	//   - Everything else will be split into quads by creating a new vertex in the center, then adding a new vertex to 
+	//       each original perimeter edge and connecting each original vertex to it's new neighbors and the center
+	//
+	// NOTE: Concave polygons will yield bad results
+	//
+
+	// Create a new vertex in the center of each incoming polygon
+	static TArray<FVertexID> PolygonCenterVertices;
+	PolygonCenterVertices.Reset();
+	{
+		static TArray<FVertexToCreate> VerticesToCreate;
+		VerticesToCreate.Reset();
+		for( const FPolygonRef PolygonRef : PolygonRefs )
+		{
+			const int32 PerimeterEdgeCount = GetPolygonPerimeterEdgeCount( PolygonRef );
+			if( TriangleTessellationMode == ETriangleTessellationMode::ThreeTriangles || PerimeterEdgeCount > 3 )
+			{
+				// Find the center of this polygon
+				const FVector PolygonCenter = ComputePolygonCenter( PolygonRef );
+
+				FVertexToCreate& VertexToCreate = *new( VerticesToCreate ) FVertexToCreate();
+				VertexToCreate.VertexAttributes.Attributes.Add( FMeshElementAttributeData( UEditableMeshAttribute::VertexPosition(), 0, PolygonCenter ) );
+			}
+		}
+
+		this->CreateVertices( VerticesToCreate, /* Out */ PolygonCenterVertices );
+	}
+
+
+	// Split all of the edges of the original polygons (except triangles).  Remember, some edges may be shared between
+	// the incoming polygons so we'll keep track of that and make sure not to split them again.  
+	{
+		static TSet<FEdgeID> EdgesToSplit;
+		EdgesToSplit.Reset();
+
+		for( int32 PolygonNumber = 0; PolygonNumber < PolygonRefs.Num(); ++PolygonNumber )
+		{
+			const FPolygonRef PolygonRef = PolygonRefs[ PolygonNumber ];
+
+			static TArray<FEdgeID> PerimeterEdgeIDs;
+			GetPolygonPerimeterEdges( PolygonRef, /* Out */ PerimeterEdgeIDs );
+			
+			if( TriangleTessellationMode == ETriangleTessellationMode::FourTriangles || PerimeterEdgeIDs.Num() > 3 )
+			{
+				for( const FEdgeID PerimeterEdgeID : PerimeterEdgeIDs )
+				{
+					EdgesToSplit.Add( PerimeterEdgeID );
+				}
+			}
+		}
+
+		for( const FEdgeID EdgeID : EdgesToSplit )
+		{
+			// Split the edge
+			static TArray<float> Splits;
+			Splits.SetNumUninitialized( 1 );
+			Splits[ 0 ] = 0.5f;
+
+			static TArray<FVertexID> NewVertexIDsFromSplit;
+			this->SplitEdge( EdgeID, Splits, /* Out */ NewVertexIDsFromSplit );
+			check( NewVertexIDsFromSplit.Num() == 1 );
+		}
+	}
+
+
+	// We'll now define the new polygons to be created.
+	static TArray<FPolygonToCreate> PolygonsToCreate;
+	PolygonsToCreate.Reset();
+
+	int32 PolygonWithNewCenterVertexNumber = 0;
+	for( int32 PolygonNumber = 0; PolygonNumber < PolygonRefs.Num(); ++PolygonNumber )
+	{
+		const FPolygonRef PolygonRef = PolygonRefs[ PolygonNumber ];
+
+		const int32 PerimeterEdgeCount = GetPolygonPerimeterEdgeCount( PolygonRef );
+
+		FVertexID PolygonCenterVertexID = FVertexID::Invalid;
+		if( TriangleTessellationMode == ETriangleTessellationMode::ThreeTriangles || PerimeterEdgeCount > 6 )
+		{
+			PolygonCenterVertexID = PolygonCenterVertices[ PolygonWithNewCenterVertexNumber++ ];
+		}
+
+		// Don't bother with triangles, because we'll simply connect the original three vertices to a new
+		// center position to tessellate those.
+		if( PerimeterEdgeCount > 6 )
+		{
+			static TArray<FVertexID> PerimeterVertexIDs;
+			GetPolygonPerimeterVertices( PolygonRef, /* Out */ PerimeterVertexIDs );
+
+			const int32 PerimeterVertexCount = PerimeterEdgeCount;
+			const int32 OriginalPerimeterEdgeCount = PerimeterEdgeCount / 2;
+			for( int32 OriginalPerimeterEdgeNumber = 0; OriginalPerimeterEdgeNumber < OriginalPerimeterEdgeCount; ++OriginalPerimeterEdgeNumber )
+			{
+				const int32 CurrentVertexNumber = OriginalPerimeterEdgeNumber * 2;
+				const int32 PreviousVertexNumber = ( ( CurrentVertexNumber - 1 ) + PerimeterVertexCount ) % PerimeterVertexCount;
+				const int32 NextVertexNumber = ( CurrentVertexNumber + 1 ) % PerimeterVertexCount;
+
+				FPolygonToCreate& PolygonToCreate = *new( PolygonsToCreate ) FPolygonToCreate();
+				PolygonToCreate.SectionID = PolygonRef.SectionID;
+
+				for( int32 QuadVertexNumber = 0; QuadVertexNumber < 4; ++QuadVertexNumber )
+				{
+					FVertexAndAttributes& VertexAndAttributes = *new( PolygonToCreate.PerimeterVertices ) FVertexAndAttributes();
+
+					int32 PerimeterVertexNumber = INDEX_NONE;
+					switch( QuadVertexNumber )
+					{
+						case 0:
+							PerimeterVertexNumber = PreviousVertexNumber;
+							break;
+
+						case 1:
+							PerimeterVertexNumber = CurrentVertexNumber;
+							break;
+
+						case 2:
+							PerimeterVertexNumber = NextVertexNumber;
+							break;
+
+						case 3:
+							PerimeterVertexNumber = INDEX_NONE;	// The center vertex!
+							break;
+
+						default:
+							check( 0 );
+					}
+
+
+					if( PerimeterVertexNumber == INDEX_NONE )
+					{
+						VertexAndAttributes.VertexID = PolygonCenterVertexID;
+
+						// Generate interpolated UVs and vertex colors for the new vertex in the center
+						{
+							const FVector CenterVertexPosition = GetVertexAttribute( PolygonCenterVertexID, UEditableMeshAttribute::VertexPosition(), 0 );
+
+							static TArray<int32> PerimeterVertexIndices;
+							FVector TriangleVertexWeights;
+							if( ComputeBarycentricWeightForPointOnPolygon( PolygonRef, CenterVertexPosition, /* Out */ PerimeterVertexIndices, /* Out */ TriangleVertexWeights ) )
+							{
+								for( int32 TextureCoordinateIndex = 0; TextureCoordinateIndex < TextureCoordinateCount; ++TextureCoordinateIndex )
+								{
+									VertexAndAttributes.PolygonVertexAttributes.Attributes.Add( FMeshElementAttributeData(
+										UEditableMeshAttribute::VertexTextureCoordinate(),
+										TextureCoordinateIndex,
+										TriangleVertexWeights.X * GetPolygonPerimeterVertexAttribute( PolygonRef, PerimeterVertexIndices[ 0 ], UEditableMeshAttribute::VertexTextureCoordinate(), TextureCoordinateIndex ) +
+										TriangleVertexWeights.Y * GetPolygonPerimeterVertexAttribute( PolygonRef, PerimeterVertexIndices[ 1 ], UEditableMeshAttribute::VertexTextureCoordinate(), TextureCoordinateIndex ) +
+										TriangleVertexWeights.Z * GetPolygonPerimeterVertexAttribute( PolygonRef, PerimeterVertexIndices[ 2 ], UEditableMeshAttribute::VertexTextureCoordinate(), TextureCoordinateIndex ) ) );
+								}
+
+								VertexAndAttributes.PolygonVertexAttributes.Attributes.Add( FMeshElementAttributeData(
+									UEditableMeshAttribute::VertexColor(),
+									0,
+									TriangleVertexWeights.X * GetPolygonPerimeterVertexAttribute( PolygonRef, PerimeterVertexIndices[ 0 ], UEditableMeshAttribute::VertexColor(), 0 ) +
+									TriangleVertexWeights.Y * GetPolygonPerimeterVertexAttribute( PolygonRef, PerimeterVertexIndices[ 1 ], UEditableMeshAttribute::VertexColor(), 0 ) +
+									TriangleVertexWeights.Z * GetPolygonPerimeterVertexAttribute( PolygonRef, PerimeterVertexIndices[ 2 ], UEditableMeshAttribute::VertexColor(), 0 ) ) );
+							}
+						}
+					}
+					else
+					{
+						VertexAndAttributes.VertexID = PerimeterVertexIDs[ PerimeterVertexNumber ];
+
+						// Transfer over UVs and vertex colors from the original polygon that we're replacing
+						{
+							for( int32 TextureCoordinateIndex = 0; TextureCoordinateIndex < TextureCoordinateCount; ++TextureCoordinateIndex )
+							{
+								VertexAndAttributes.PolygonVertexAttributes.Attributes.Add( FMeshElementAttributeData(
+									UEditableMeshAttribute::VertexTextureCoordinate(),
+									TextureCoordinateIndex,
+									GetPolygonPerimeterVertexAttribute( PolygonRef, PerimeterVertexNumber, UEditableMeshAttribute::VertexTextureCoordinate(), TextureCoordinateIndex ) ) );
+							}
+
+							VertexAndAttributes.PolygonVertexAttributes.Attributes.Add( FMeshElementAttributeData(
+								UEditableMeshAttribute::VertexColor(),
+								0,
+								GetPolygonPerimeterVertexAttribute( PolygonRef, PerimeterVertexNumber, UEditableMeshAttribute::VertexColor(), 0 ) ) );
+						}
+					}
+				}
+			}
+		}
+		else
+		{
+			if( TriangleTessellationMode == ETriangleTessellationMode::ThreeTriangles )
+			{
+				// Define the three new triangles for the original tessellated triangle
+				for( int32 PerimeterEdgeNumber = 0; PerimeterEdgeNumber < 3; ++PerimeterEdgeNumber )
+				{
+					bool bIsEdgeWindingReversedForPolygon;
+					const FEdgeID EdgeID = GetPolygonPerimeterEdge( PolygonRef, PerimeterEdgeNumber, /* Out */ bIsEdgeWindingReversedForPolygon );
+
+					FVertexID EdgeVertexID0, EdgeVertexID1;
+					GetEdgeVertices( EdgeID, /* Out */ EdgeVertexID0, /* Out */ EdgeVertexID1 );
+
+					FPolygonToCreate& PolygonToCreate = *new( PolygonsToCreate ) FPolygonToCreate();
+					PolygonToCreate.SectionID = PolygonRef.SectionID;
+
+					for( int32 TriangleVertexNumber = 0; TriangleVertexNumber < 3; ++TriangleVertexNumber )
+					{
+						FVertexAndAttributes& VertexAndAttributes = *new( PolygonToCreate.PerimeterVertices ) FVertexAndAttributes();
+						switch( TriangleVertexNumber )
+						{
+							case 0:
+								VertexAndAttributes.VertexID = bIsEdgeWindingReversedForPolygon ? EdgeVertexID1 : EdgeVertexID0;
+								break;
+
+							case 1:
+								VertexAndAttributes.VertexID = PolygonCenterVertexID;
+								break;
+
+							case 2:
+								VertexAndAttributes.VertexID = bIsEdgeWindingReversedForPolygon ? EdgeVertexID0 : EdgeVertexID1;
+								break;
+
+							default:
+								check( 0 );
+						}
+
+						if( VertexAndAttributes.VertexID == PolygonCenterVertexID )
+						{
+							// Generate interpolated UVs and vertex colors for the new vertex in the center
+							{
+								const FVector CenterVertexPosition = GetVertexAttribute( PolygonCenterVertexID, UEditableMeshAttribute::VertexPosition(), 0 );
+
+								static TArray<int32> PerimeterVertexIndices;
+								FVector TriangleVertexWeights;
+								if( ComputeBarycentricWeightForPointOnPolygon( PolygonRef, CenterVertexPosition, /* Out */ PerimeterVertexIndices, /* Out */ TriangleVertexWeights ) )
+								{
+									for( int32 TextureCoordinateIndex = 0; TextureCoordinateIndex < TextureCoordinateCount; ++TextureCoordinateIndex )
+									{
+										VertexAndAttributes.PolygonVertexAttributes.Attributes.Add( FMeshElementAttributeData(
+											UEditableMeshAttribute::VertexTextureCoordinate(),
+											TextureCoordinateIndex,
+											TriangleVertexWeights.X * GetPolygonPerimeterVertexAttribute( PolygonRef, PerimeterVertexIndices[ 0 ], UEditableMeshAttribute::VertexTextureCoordinate(), TextureCoordinateIndex ) +
+											TriangleVertexWeights.Y * GetPolygonPerimeterVertexAttribute( PolygonRef, PerimeterVertexIndices[ 1 ], UEditableMeshAttribute::VertexTextureCoordinate(), TextureCoordinateIndex ) +
+											TriangleVertexWeights.Z * GetPolygonPerimeterVertexAttribute( PolygonRef, PerimeterVertexIndices[ 2 ], UEditableMeshAttribute::VertexTextureCoordinate(), TextureCoordinateIndex ) ) );
+									}
+
+									VertexAndAttributes.PolygonVertexAttributes.Attributes.Add( FMeshElementAttributeData(
+										UEditableMeshAttribute::VertexColor(),
+										0,
+										TriangleVertexWeights.X * GetPolygonPerimeterVertexAttribute( PolygonRef, PerimeterVertexIndices[ 0 ], UEditableMeshAttribute::VertexColor(), 0 ) +
+										TriangleVertexWeights.Y * GetPolygonPerimeterVertexAttribute( PolygonRef, PerimeterVertexIndices[ 1 ], UEditableMeshAttribute::VertexColor(), 0 ) +
+										TriangleVertexWeights.Z * GetPolygonPerimeterVertexAttribute( PolygonRef, PerimeterVertexIndices[ 2 ], UEditableMeshAttribute::VertexColor(), 0 ) ) );
+								}
+							}
+						}
+						else
+						{
+							// Transfer over UVs and vertex colors from the original polygon that we're replacing
+							{
+								const int32 VertexNumber = FindPolygonPerimeterVertexNumberForVertex( PolygonRef, VertexAndAttributes.VertexID );
+								check( VertexNumber != INDEX_NONE );
+
+								for( int32 TextureCoordinateIndex = 0; TextureCoordinateIndex < TextureCoordinateCount; ++TextureCoordinateIndex )
+								{
+									VertexAndAttributes.PolygonVertexAttributes.Attributes.Add( FMeshElementAttributeData(
+										UEditableMeshAttribute::VertexTextureCoordinate(),
+										TextureCoordinateIndex,
+										GetPolygonPerimeterVertexAttribute( PolygonRef, VertexNumber, UEditableMeshAttribute::VertexTextureCoordinate(), TextureCoordinateIndex ) ) );
+								}
+
+								VertexAndAttributes.PolygonVertexAttributes.Attributes.Add( FMeshElementAttributeData(
+									UEditableMeshAttribute::VertexColor(),
+									0,
+									GetPolygonPerimeterVertexAttribute( PolygonRef, VertexNumber, UEditableMeshAttribute::VertexColor(), 0 ) ) );
+							}
+						}
+					}
+				}
+			}
+			else if( ensure( TriangleTessellationMode == ETriangleTessellationMode::FourTriangles ) )
+			{
+				// Define the four new triangles for the original tessellated triangle.  One triangle will go in
+				// the center, connecting the three new vertices that we created between each original edge.  The
+				// other three triangles will go in the corners of the original triangle.
+
+				static TArray<FVertexID> PerimeterVertexIDs;
+				GetPolygonPerimeterVertices( PolygonRef, PerimeterVertexIDs );
+				check( PerimeterVertexIDs.Num() == 6 );	// We split the triangle's 3 edges earlier, so we must have six edges now
+
+				// Define the new center triangle
+				{
+					FPolygonToCreate& PolygonToCreate = *new( PolygonsToCreate ) FPolygonToCreate();
+					PolygonToCreate.SectionID = PolygonRef.SectionID;
+
+					for( int32 OriginalVertexNumber = 0; OriginalVertexNumber < 3; ++OriginalVertexNumber )
+					{
+						const int32 VertexNumber = ( OriginalVertexNumber * 2 + 1 ) % PerimeterVertexIDs.Num();
+
+						FVertexAndAttributes& VertexAndAttributes = *new( PolygonToCreate.PerimeterVertices ) FVertexAndAttributes();
+						VertexAndAttributes.VertexID = PerimeterVertexIDs[ VertexNumber ];
+
+						// Transfer over UVs and vertex colors from the original polygon that we're replacing
+						{
+							for( int32 TextureCoordinateIndex = 0; TextureCoordinateIndex < TextureCoordinateCount; ++TextureCoordinateIndex )
+							{
+								VertexAndAttributes.PolygonVertexAttributes.Attributes.Add( FMeshElementAttributeData(
+									UEditableMeshAttribute::VertexTextureCoordinate(),
+									TextureCoordinateIndex,
+									GetPolygonPerimeterVertexAttribute( PolygonRef, VertexNumber, UEditableMeshAttribute::VertexTextureCoordinate(), TextureCoordinateIndex ) ) );
+							}
+
+							VertexAndAttributes.PolygonVertexAttributes.Attributes.Add( FMeshElementAttributeData(
+								UEditableMeshAttribute::VertexColor(),
+								0,
+								GetPolygonPerimeterVertexAttribute( PolygonRef, VertexNumber, UEditableMeshAttribute::VertexColor(), 0 ) ) );
+						}
+					}
+				}
+
+				// Define the three corner triangles
+				for( int32 OriginalEdgeNumber = 0; OriginalEdgeNumber < 3; ++OriginalEdgeNumber )
+				{
+					const int32 CurrentVertexNumber = OriginalEdgeNumber * 2;
+					const int32 PreviousVertexNumber = ( ( CurrentVertexNumber - 1 ) + PerimeterVertexIDs.Num() ) % PerimeterVertexIDs.Num();
+					const int32 NextVertexNumber = ( CurrentVertexNumber + 1 ) % PerimeterVertexIDs.Num();
+
+					FPolygonToCreate& PolygonToCreate = *new( PolygonsToCreate ) FPolygonToCreate();
+					PolygonToCreate.SectionID = PolygonRef.SectionID;
+
+					for( int32 TriangleVertexNumber = 0; TriangleVertexNumber < 3; ++TriangleVertexNumber )
+					{
+						int32 VertexNumber = INDEX_NONE;
+						switch( TriangleVertexNumber )
+						{
+							case 0:
+								VertexNumber = PreviousVertexNumber;
+								break;
+
+							case 1:
+								VertexNumber = CurrentVertexNumber;
+								break;
+
+							case 2:
+								VertexNumber = NextVertexNumber;
+								break;
+
+							default:
+								check( 0 );
+						}
+
+						FVertexAndAttributes& VertexAndAttributes = *new( PolygonToCreate.PerimeterVertices ) FVertexAndAttributes();
+						VertexAndAttributes.VertexID = PerimeterVertexIDs[ VertexNumber ];
+
+						{
+							// Transfer over UVs and vertex colors from the original polygon that we're replacing
+							{
+								for( int32 TextureCoordinateIndex = 0; TextureCoordinateIndex < TextureCoordinateCount; ++TextureCoordinateIndex )
+								{
+									VertexAndAttributes.PolygonVertexAttributes.Attributes.Add( FMeshElementAttributeData(
+										UEditableMeshAttribute::VertexTextureCoordinate(),
+										TextureCoordinateIndex,
+										GetPolygonPerimeterVertexAttribute( PolygonRef, VertexNumber, UEditableMeshAttribute::VertexTextureCoordinate(), TextureCoordinateIndex ) ) );
+								}
+
+								VertexAndAttributes.PolygonVertexAttributes.Attributes.Add( FMeshElementAttributeData(
+									UEditableMeshAttribute::VertexColor(),
+									0,
+									GetPolygonPerimeterVertexAttribute( PolygonRef, VertexNumber, UEditableMeshAttribute::VertexColor(), 0 ) ) );
+							}
+						}
+					}
+				}
+			}
+		}
+	}
+
+
+	// Delete the original polygons
+	{
+		const bool bDeleteOrphanedEdges = false;	// No need to delete orphans, because this function won't orphan anything
+		const bool bDeleteOrphanedVertices = false;
+		const bool bDeleteEmptySections = false;
+		DeletePolygons( PolygonRefs, bDeleteOrphanedEdges, bDeleteOrphanedVertices, bDeleteEmptySections );
+	}
+
+
+	// Create all of the new polygons for the tessellated representation of the original polygons
+	{
+		static TArray<FEdgeID> NewEdgeIDs;
+		CreatePolygons( PolygonsToCreate, /* Out */ OutNewPolygonRefs, /* Out */ NewEdgeIDs );
+	}
+
+
+	// Generate normals for all of the new polygons
+	GenerateNormalsAndTangentsForPolygons( OutNewPolygonRefs );
+}
+
+
+void UEditableMesh::SetTextureCoordinateCount(int32 NumTexCoords)
+{
+	TextureCoordinateCount = FMath::Max(NumTexCoords, 0);
+}
+
+
+void UEditableMesh::QuadrangulateMesh( TArray<FPolygonRef>& NewPolygonRefs )
+{
+	// Uses the first two steps of the algorithm described by
+	// http://www.lirmm.fr/~beniere/ArticlesPersos/GRAPP10_Beniere_Final.pdf
+
+	// Tweakable parameters affecting how quadrangulate works
+	const float CosAngleThreshold = 0.984f;		// about 10 degrees
+	const bool bKeepHardEdges = true;
+	const bool bKeepTextureBorder = true;
+	const bool bKeepColorBorder = true;
+
+	NewPolygonRefs.Reset();
+
+	static TArray<FPolygonRef> PolygonRefs;
+	PolygonRefs.Reset();
+
+	static TSet<FPolygonRef> DeletedPolygonRefs;
+	DeletedPolygonRefs.Reset();
+
+	// Get a list of all polygon refs in the mesh
+	{
+		const int32 MaxSectionIndex = GetSectionArraySize();
+		for( int32 SectionIndex = 0; SectionIndex < MaxSectionIndex; ++SectionIndex )
+		{
+			const FSectionID SectionID( SectionIndex );
+			if( IsValidSection( SectionID ) )
+			{
+				const int32 MaxPolygonIndex = GetPolygonArraySize( SectionID );
+				for( int32 PolygonIndex = 0; PolygonIndex < MaxPolygonIndex; ++PolygonIndex )
+				{
+					const FPolygonRef PolygonRef( SectionID, FPolygonID( PolygonIndex ) );
+					if( IsValidPolygon( PolygonRef ) )
+					{
+						PolygonRefs.Add( PolygonRef );
+					}
+				}
+			}
+		}
+	}
+
+	// This represents an adjacent triangle which can be merged to a quadrilateral, and an assigned score based on the 'quality' of the resulting quadrilateral.
+	struct FAdjacentPolygon
+	{
+		// PolygonRef of the adjacent triangle. This object is keyed on 'our' PolygonRef.
+		FPolygonRef PolygonRef;
+
+		// List of the vertices which form the merged quadrilateral - (PolygonRef, PolygonPerimeterIndex)
+		TTuple<FPolygonRef, int32> Vertices[ 4 ];
+
+		// 'Quality' of the quadrilateral (internal angles closer to 90 degrees are better)
+		float Score;
+
+		FAdjacentPolygon() {}
+
+		FAdjacentPolygon( FPolygonRef InPolygonRef, FPolygonRef InAdjacentPolygonRef, int32 InVertex0, int32 InVertex1, int32 InVertex2, int32 InVertex3, float InScore )
+		{
+			PolygonRef = InAdjacentPolygonRef;
+			Vertices[ 0 ] = MakeTuple( InPolygonRef, InVertex0 );
+			Vertices[ 1 ] = MakeTuple( InPolygonRef, InVertex1 );
+			Vertices[ 2 ] = MakeTuple( InAdjacentPolygonRef, InVertex2 );
+			Vertices[ 3 ] = MakeTuple( InAdjacentPolygonRef, InVertex3 );
+			Score = InScore;
+		}
+	};
+
+
+	// This represents a list of adjacent polygons, ordered by score.
+	// Since we are only connecting triangles, there are a maximum of three adjacent polygons.
+	struct FAdjacentPolygons
+	{
+		enum { MaxAdjacentPolygons = 3 };
+
+		FAdjacentPolygon AdjacentPolygons[ MaxAdjacentPolygons ];
+		int32 NumAdjacentPolygons;
+
+		FAdjacentPolygons()
+		{
+			NumAdjacentPolygons = 0;
+		}
+
+		void Add( const FAdjacentPolygon& AdjacentPolygon )
+		{
+			check( NumAdjacentPolygons < MaxAdjacentPolygons );
+			int32 InsertIndex = 0;
+			for( int32 Index = 0; Index < NumAdjacentPolygons; ++Index )
+			{
+				if( AdjacentPolygon.Score > AdjacentPolygons[ Index ].Score )
+				{
+					InsertIndex++;
+				}
+				else
+				{
+					break;
+				}
+			}
+
+			for( int32 Index = NumAdjacentPolygons; Index > InsertIndex; --Index )
+			{
+				AdjacentPolygons[ Index ] = AdjacentPolygons[ Index - 1 ];
+			}
+			AdjacentPolygons[ InsertIndex ] = AdjacentPolygon;
+			NumAdjacentPolygons++;
+		}
+
+		const FAdjacentPolygon& GetBestAdjacentPolygon() const
+		{
+			check( NumAdjacentPolygons > 0 );
+			return AdjacentPolygons[ 0 ];
+		}
+		
+		bool Remove( const FPolygonRef PolygonRef )
+		{
+			for( int32 Index = 0; Index < NumAdjacentPolygons; ++Index )
+			{
+				if( AdjacentPolygons[ Index ].PolygonRef == PolygonRef )
+				{
+					for( int32 CopyIndex = Index + 1; CopyIndex < NumAdjacentPolygons; ++CopyIndex )
+					{
+						AdjacentPolygons[ CopyIndex - 1 ] = AdjacentPolygons[ CopyIndex ];
+					}
+
+					NumAdjacentPolygons--;
+					return true;
+				}
+			}
+
+			return false;
+		}
+
+		bool Contains( const FPolygonRef PolygonRef ) const
+		{
+			for( int32 Index = 0; Index < NumAdjacentPolygons; ++Index )
+			{
+				if( AdjacentPolygons[ Index ].PolygonRef == PolygonRef )
+				{
+					return true;
+				}
+			}
+
+			return false;
+		}
+
+		FPolygonRef GetPolygonRef( int32 Index ) const
+		{
+			check( Index < NumAdjacentPolygons );
+			return AdjacentPolygons[ Index ].PolygonRef;
+		}
+
+		int32 Num() const { return NumAdjacentPolygons; }
+		bool IsValid() const { return NumAdjacentPolygons > 0; }
+	};
+
+
+	// Build list of valid adjacent triangle pairs, and assign a score based on the quality of the quadrilateral they form
+
+	static TMap<FPolygonRef, FAdjacentPolygons> AdjacentPolygonsMap;
+	AdjacentPolygonsMap.Reset();
+
+	FPolygonRef PolygonRefToMerge1 = FPolygonRef::Invalid;
+	{
+		float BestScore = TNumericLimits<float>::Max();
+
+		for( const FPolygonRef PolygonRef : PolygonRefs )
+		{
+			// If it's not a triangle, don't consider this polygon at all
+			if( GetPolygonPerimeterEdgeCount( PolygonRef ) != 3 )
+			{
+				continue;
+			}
+
+			// We're only interested in adjacent triangles which are nearly coplanar; get the normal so we can compare it with the adjacent polygons' normals
+			const FVector PolygonNormal = ComputePolygonNormal( PolygonRef );
+
+			// Go round the edge considering all adjacent polygons, looking for valid pairs and assigning a quality score (lower is better)
+			for( int32 PerimeterEdgeIndex = 0; PerimeterEdgeIndex < 3; ++PerimeterEdgeIndex )
+			{
+				bool bOutEdgeWindingIsReversedForPolygon;
+				const FEdgeID PerimeterEdgeID = GetPolygonPerimeterEdge( PolygonRef, PerimeterEdgeIndex, bOutEdgeWindingIsReversedForPolygon );
+
+				const bool bIsSoftEdge = FMath::IsNearlyZero( GetEdgeAttribute( PerimeterEdgeID, UEditableMeshAttribute::EdgeIsHard(), 0 ).X );
+				if( !bKeepHardEdges || bIsSoftEdge )
+				{
+					const FPolygonRef AdjacentPolygonRef = [ this, PerimeterEdgeID, PolygonRef ]() -> FPolygonRef
+					{
+						const int32 EdgeConnectedPolygonCount = GetEdgeConnectedPolygonCount( PerimeterEdgeID );
+
+						// Only interested in edges with exactly two connected polygons
+						if( EdgeConnectedPolygonCount == 2 )
+						{
+							for( int32 EdgeConnectedPolygonIndex = 0; EdgeConnectedPolygonIndex < 2; ++EdgeConnectedPolygonIndex )
+							{
+								const FPolygonRef EdgeConnectedPolygonRef = GetEdgeConnectedPolygon( PerimeterEdgeID, EdgeConnectedPolygonIndex );
+								if( EdgeConnectedPolygonRef != PolygonRef )
+								{
+									return GetPolygonPerimeterEdgeCount( EdgeConnectedPolygonRef ) == 3 ? EdgeConnectedPolygonRef : FPolygonRef::Invalid;
+								}
+							}
+						}
+
+						return FPolygonRef::Invalid;
+					}();
+
+					if( AdjacentPolygonRef != FPolygonRef::Invalid )
+					{
+						FAdjacentPolygons* AdjacentPolygons = AdjacentPolygonsMap.Find( PolygonRef );
+						if( !AdjacentPolygons || !AdjacentPolygons->Contains( AdjacentPolygonRef ) )
+						{
+							const FVector AdjacentPolygonNormal = ComputePolygonNormal( AdjacentPolygonRef );
+							const float AdjacentPolygonDot = FVector::DotProduct( PolygonNormal, AdjacentPolygonNormal );
+
+							if( AdjacentPolygonDot >= CosAngleThreshold )
+							{
+								// Found a valid triangle pair whose interplanar angle is sufficiently shallow;
+								// now calculate a score according to the internal angles of the resulting quad
+
+								// We consider points on the two triangles' perimeters which form the quadrilateral.
+								// If the shared edge of the adjacent triangles falls on perimeter vertex N1 of triangle 1,
+								// and perimeter vertex N2 of triangle 2, then the points we consider are:
+								//
+								// (triangle 1, point N1 - 1)
+								// (triangle 1, point N1)
+								// (triangle 2, point N2 + 1)
+								// (triangle 2, point N2 - 1)
+								//
+								// or, from the perspective of the other triangle:
+								//
+								// (triangle 2, point N2 + 1)
+								// (triangle 2, point N2 - 1)
+								// (triangle 1, point N1 - 1)
+								// (triangle 1, point N1)
+
+								const int32 PrevPerimeterEdgeIndex = ( PerimeterEdgeIndex + 2 ) % 3;
+								const int32 NextPerimeterEdgeIndex = ( PerimeterEdgeIndex + 1 ) % 3;
+
+								const FVertexID SharedVertexID = GetPolygonPerimeterVertex( PolygonRef, PerimeterEdgeIndex );
+								const int32 AdjacentPerimeterEdgeIndex = FindPolygonPerimeterVertexNumberForVertex( AdjacentPolygonRef, SharedVertexID );
+
+								const int32 PrevAdjacentPerimeterEdgeIndex = ( AdjacentPerimeterEdgeIndex + 2 ) % 3;
+								const int32 NextAdjacentPerimeterEdgeIndex = ( AdjacentPerimeterEdgeIndex + 1 ) % 3;
+
+								// If we wish to maintain texture seams, compare the vertex UVs of each polygon's vertex along the edge to be removed, and only proceed if they are equal
+								bool bTextureCoordinatesEqual = true;
+								if( bKeepTextureBorder )
+								{
+									for( int32 TextureCoordinateIndex = 0; TextureCoordinateIndex < TextureCoordinateCount; ++TextureCoordinateIndex )
+									{
+										const FVector4 UVPolygon1Vertex1 = GetPolygonPerimeterVertexAttribute( PolygonRef, PerimeterEdgeIndex, UEditableMeshAttribute::VertexTextureCoordinate(), TextureCoordinateIndex );
+										const FVector4 UVPolygon2Vertex1 = GetPolygonPerimeterVertexAttribute( AdjacentPolygonRef, AdjacentPerimeterEdgeIndex, UEditableMeshAttribute::VertexTextureCoordinate(), TextureCoordinateIndex );
+										const FVector4 UVPolygon1Vertex2 = GetPolygonPerimeterVertexAttribute( PolygonRef, NextPerimeterEdgeIndex, UEditableMeshAttribute::VertexTextureCoordinate(), TextureCoordinateIndex );
+										const FVector4 UVPolygon2Vertex2 = GetPolygonPerimeterVertexAttribute( AdjacentPolygonRef, PrevAdjacentPerimeterEdgeIndex, UEditableMeshAttribute::VertexTextureCoordinate(), TextureCoordinateIndex );
+										if( UVPolygon1Vertex1 != UVPolygon2Vertex1 || UVPolygon1Vertex2 != UVPolygon2Vertex2 )
+										{
+											bTextureCoordinatesEqual = false;
+											break;
+										}
+									}
+								}
+
+								// If we wish to maintain color seams, compare the vertex colors of each polygon's vertex along the edge to be removed, and only proceed if they are equal
+								bool bColorsEqual = true;
+								if( bKeepColorBorder )
+								{
+									const FVector4 ColorPolygon1Vertex1 = GetPolygonPerimeterVertexAttribute( PolygonRef, PerimeterEdgeIndex, UEditableMeshAttribute::VertexColor(), 0 );
+									const FVector4 ColorPolygon2Vertex1 = GetPolygonPerimeterVertexAttribute( AdjacentPolygonRef, AdjacentPerimeterEdgeIndex, UEditableMeshAttribute::VertexColor(), 0 );
+									const FVector4 ColorPolygon1Vertex2 = GetPolygonPerimeterVertexAttribute( PolygonRef, NextPerimeterEdgeIndex, UEditableMeshAttribute::VertexColor(), 0 );
+									const FVector4 ColorPolygon2Vertex2 = GetPolygonPerimeterVertexAttribute( AdjacentPolygonRef, PrevAdjacentPerimeterEdgeIndex, UEditableMeshAttribute::VertexColor(), 0 );
+									if( ColorPolygon1Vertex1 != ColorPolygon2Vertex1 || ColorPolygon1Vertex2 != ColorPolygon2Vertex2 )
+									{
+										bColorsEqual = false;
+									}
+								}
+
+								// If we are maintaining texture or color seams, only merge the polygons if the appropriate attributes are equal for each polygon
+								if( ( !bKeepTextureBorder || bTextureCoordinatesEqual ) && ( !bKeepColorBorder || bColorsEqual ) )
+								{
+									const FVertexID V0 = GetPolygonPerimeterVertex( PolygonRef, PrevPerimeterEdgeIndex );
+									const FVertexID V1 = GetPolygonPerimeterVertex( PolygonRef, PerimeterEdgeIndex );
+									const FVertexID V2 = GetPolygonPerimeterVertex( AdjacentPolygonRef, NextAdjacentPerimeterEdgeIndex );
+									const FVertexID V3 = GetPolygonPerimeterVertex( AdjacentPolygonRef, PrevAdjacentPerimeterEdgeIndex );
+									check( V3 == GetPolygonPerimeterVertex( PolygonRef, NextPerimeterEdgeIndex ) );
+
+									const FVector P0 = GetVertexAttribute( V0, UEditableMeshAttribute::VertexPosition(), 0 );
+									const FVector P1 = GetVertexAttribute( V1, UEditableMeshAttribute::VertexPosition(), 0 );
+									const FVector P2 = GetVertexAttribute( V2, UEditableMeshAttribute::VertexPosition(), 0 );
+									const FVector P3 = GetVertexAttribute( V3, UEditableMeshAttribute::VertexPosition(), 0 );
+
+									const FVector D01 = ( P1 - P0 ).GetSafeNormal();
+									const FVector D12 = ( P2 - P1 ).GetSafeNormal();
+									const FVector D23 = ( P3 - P2 ).GetSafeNormal();
+									const FVector D30 = ( P0 - P3 ).GetSafeNormal();
+
+									// Calculate a score based on the internal angles of the quadrilateral and the interplanar angle.
+									// Internal angles close to 90 degrees, and an interplanar angle close to 180 degrees are ideal.
+									const float Score =
+										FMath::Abs( HALF_PI - FMath::Acos( FVector::DotProduct( -D30, D01 ) ) ) +
+										FMath::Abs( HALF_PI - FMath::Acos( FVector::DotProduct( -D01, D12 ) ) ) +
+										FMath::Abs( HALF_PI - FMath::Acos( FVector::DotProduct( -D12, D23 ) ) ) +
+										FMath::Abs( HALF_PI - FMath::Acos( FVector::DotProduct( -D23, D30 ) ) ) +
+										FMath::Acos( AdjacentPolygonDot );
+
+									if( Score < BestScore )
+									{
+										BestScore = Score;
+										PolygonRefToMerge1 = PolygonRef;
+									}
+
+									// Add to a list of adjacent polygons, sorted by score
+									FAdjacentPolygons& AdjacentPolygons1 = AdjacentPolygonsMap.FindOrAdd( PolygonRef );
+									AdjacentPolygons1.Add( FAdjacentPolygon(
+										PolygonRef,
+										AdjacentPolygonRef,
+										PrevPerimeterEdgeIndex,
+										PerimeterEdgeIndex,
+										NextAdjacentPerimeterEdgeIndex,
+										PrevAdjacentPerimeterEdgeIndex,
+										Score ) );
+
+									// And perform the corresponding operation the other way round
+									FAdjacentPolygons& AdjacentPolygons2 = AdjacentPolygonsMap.FindOrAdd( AdjacentPolygonRef );
+									check( !AdjacentPolygons2.Contains( PolygonRef ) );
+									AdjacentPolygons2.Add( FAdjacentPolygon(
+										AdjacentPolygonRef,
+										PolygonRef,
+										NextAdjacentPerimeterEdgeIndex,
+										PrevAdjacentPerimeterEdgeIndex,
+										PrevPerimeterEdgeIndex,
+										PerimeterEdgeIndex,
+										Score ) );
+								}
+							}
+						}
+					}
+				}
+			}
+		}
+	}
+
+	if( PolygonRefToMerge1 == FPolygonRef::Invalid )
+	{
+		return;
+	}
+
+	static TArray<FPolygonToCreate> PolygonsToCreate;
+	PolygonsToCreate.Reset();
+
+	static TArray<FPolygonRef> PolygonRefsToDelete;
+	PolygonRefsToDelete.Reset();
+
+	TSet<FPolygonRef> BoundaryPolygons;
+
+	// Propagate quadrangulated area outwards from starting polygon
+	for( ; ; )
+	{
+		FAdjacentPolygons& AdjacentPolygons1 = AdjacentPolygonsMap.FindChecked( PolygonRefToMerge1 );
+		check( AdjacentPolygons1.IsValid() );
+		const FAdjacentPolygon& AdjacentPolygon1 = AdjacentPolygons1.GetBestAdjacentPolygon();
+		const FPolygonRef PolygonRefToMerge2 = AdjacentPolygon1.PolygonRef;
+
+		FAdjacentPolygons& AdjacentPolygons2 = AdjacentPolygonsMap.FindChecked( PolygonRefToMerge2 );
+		check( AdjacentPolygons2.IsValid() );
+
+		// Create new quadrilateral
+
+		PolygonsToCreate.Emplace();
+
+		FPolygonToCreate& PolygonToCreate = PolygonsToCreate.Last();
+
+		PolygonToCreate.SectionID = PolygonRefToMerge1.SectionID;
+		PolygonToCreate.PerimeterVertices.Reset( 4 );
+
+		for( int32 VertexIndex = 0; VertexIndex < 4; ++VertexIndex )
+		{
+			PolygonToCreate.PerimeterVertices.Emplace();
+			FVertexAndAttributes& VertexAndAttributes = PolygonToCreate.PerimeterVertices.Last();
+
+			VertexAndAttributes.VertexID = GetPolygonPerimeterVertex(
+				AdjacentPolygon1.Vertices[ VertexIndex ].Get<0>(),
+				AdjacentPolygon1.Vertices[ VertexIndex ].Get<1>() );
+
+			TArray<FMeshElementAttributeData>& Attributes = VertexAndAttributes.PolygonVertexAttributes.Attributes;
+
+			for( int32 TextureCoordinateIndex = 0; TextureCoordinateIndex < TextureCoordinateCount; ++TextureCoordinateIndex )
+			{
+				Attributes.Emplace(
+					UEditableMeshAttribute::VertexTextureCoordinate(),
+					TextureCoordinateIndex,
+					GetPolygonPerimeterVertexAttribute( 
+						AdjacentPolygon1.Vertices[ VertexIndex ].Get<0>(),
+						AdjacentPolygon1.Vertices[ VertexIndex ].Get<1>(),
+						UEditableMeshAttribute::VertexTextureCoordinate(),
+						TextureCoordinateIndex ) 
+					);
+			}
+
+			Attributes.Emplace(
+				UEditableMeshAttribute::VertexColor(),
+				0,
+				GetPolygonPerimeterVertexAttribute(
+					AdjacentPolygon1.Vertices[ VertexIndex ].Get<0>(),
+					AdjacentPolygon1.Vertices[ VertexIndex ].Get<1>(),
+					UEditableMeshAttribute::VertexColor(),
+					0 )
+				);
+		}
+
+		// Specify old polygons to be deleted
+
+		check( !PolygonRefsToDelete.Contains( PolygonRefToMerge1 ) );
+		check( !PolygonRefsToDelete.Contains( PolygonRefToMerge2 ) );
+		PolygonRefsToDelete.Emplace( PolygonRefToMerge1 );
+		PolygonRefsToDelete.Emplace( PolygonRefToMerge2 );
+
+		// And remove them from the boundary set
+
+		BoundaryPolygons.Remove( PolygonRefToMerge1 );
+		BoundaryPolygons.Remove( PolygonRefToMerge2 );
+
+		// Now break connections between newly added polygons and their neighbors.
+		// If a polygon ends up with no connections, delete it entirely from the map so it is no longer considered.
+		// This happens if a polygon has been added to the quadrangulated set, or if it is an orphaned triangle which cannot be paired to anything.
+		// We defer deleting the entry from the map until we have broken all connections.
+
+		verify( AdjacentPolygons1.Remove( PolygonRefToMerge2 ) );
+		verify( AdjacentPolygons2.Remove( PolygonRefToMerge1 ) );
+
+		static TArray<FPolygonRef> AdjacentPolygonsEntryToDelete;
+		AdjacentPolygonsEntryToDelete.Reset();
+
+		for( int32 Index = 0; Index < AdjacentPolygons1.Num(); ++Index )
+		{
+			const FPolygonRef AdjacentPolygonRef = AdjacentPolygons1.GetPolygonRef( Index );
+
+			FAdjacentPolygons* OtherAdjacentPolygons = AdjacentPolygonsMap.Find( AdjacentPolygonRef );
+			if( OtherAdjacentPolygons )
+			{
+				verify( OtherAdjacentPolygons->Remove( PolygonRefToMerge1 ) );
+				if( !OtherAdjacentPolygons->IsValid() )
+				{
+					AdjacentPolygonsEntryToDelete.Add( AdjacentPolygonRef );
+				}
+				else
+				{
+					BoundaryPolygons.Add( AdjacentPolygonRef );
+				}
+			}
+		}
+
+		AdjacentPolygonsEntryToDelete.Add( PolygonRefToMerge1 );
+
+		for( int32 Index = 0; Index < AdjacentPolygons2.Num(); ++Index )
+		{
+			const FPolygonRef AdjacentPolygonRef = AdjacentPolygons2.GetPolygonRef( Index );
+
+			FAdjacentPolygons* OtherAdjacentPolygons = AdjacentPolygonsMap.Find( AdjacentPolygonRef );
+			if( OtherAdjacentPolygons )
+			{
+				verify( OtherAdjacentPolygons->Remove( PolygonRefToMerge2 ) );
+				if( !OtherAdjacentPolygons->IsValid() )
+				{
+					AdjacentPolygonsEntryToDelete.Add( AdjacentPolygonRef );
+				}
+				else
+				{
+					BoundaryPolygons.Add( AdjacentPolygonRef );
+				}
+			}
+		}
+
+		AdjacentPolygonsEntryToDelete.Add( PolygonRefToMerge2 );
+
+		// Clean up: any polygons' map entries which now have no adjacent polygons get deleted completely.
+		// This implies they have no connected neighbors which can be merged (either because they are near an edge with only unmergeable polygons nearby, or because
+		// they are in the middle of the quadrangulated area)
+
+		for( const FPolygonRef AdjacentPolygonsEntry : AdjacentPolygonsEntryToDelete )
+		{
+			AdjacentPolygonsMap.Remove( AdjacentPolygonsEntry );
+			BoundaryPolygons.Remove( AdjacentPolygonsEntry );
+		}
+
+		// Now look for the next polygon to use: it is the one with the best score from the BoundaryPolygons set.
+
+		float BestScore = TNumericLimits<float>::Max();
+		PolygonRefToMerge1 = FPolygonRef::Invalid;
+		for( const FPolygonRef BoundaryPolygon : BoundaryPolygons )
+		{
+			const FAdjacentPolygons& AdjacentPolygons = AdjacentPolygonsMap.FindChecked( BoundaryPolygon );
+			const FAdjacentPolygon& AdjacentPolygon = AdjacentPolygons.GetBestAdjacentPolygon();
+			if( AdjacentPolygon.Score < BestScore )
+			{
+				PolygonRefToMerge1 = AdjacentPolygon.PolygonRef;
+			}
+		}
+
+		// If there are still no candidates adjacent to the already quadrangulated area, choose the best candidate elsewhere.
+		// This will start a new quadrangulated area, which is grown in the same way as the last.
+
+		if( PolygonRefToMerge1 == FPolygonRef::Invalid )
+		{
+			BoundaryPolygons.Reset();
+
+			for( const auto& AdjacentPolygonsMapEntry : AdjacentPolygonsMap )
+			{
+				const FPolygonRef PolygonRef = AdjacentPolygonsMapEntry.Key;
+				const FAdjacentPolygons& AdjacentPolygons = AdjacentPolygonsMapEntry.Value;
+				check( AdjacentPolygons.IsValid() );
+
+				const FAdjacentPolygon& AdjacentPolygon = AdjacentPolygons.GetBestAdjacentPolygon();
+
+				if( AdjacentPolygon.Score < BestScore )
+				{
+					BestScore = AdjacentPolygon.Score;
+					PolygonRefToMerge1 = PolygonRef;
+				}
+			}
+		}
+
+		// If there are still no candidates, we've done as much as we can do
+
+		if( PolygonRefToMerge1 == FPolygonRef::Invalid )
+		{
+			break;
+		}
+	}
+
+	// Finally, actually change the geometry and rebuild normals/tangents
+
+	static TArray<FPolygonRef> CreatedPolygonRefs;
+	static TArray< FEdgeID> CreatedEdgeIDs;
+
+	CreatePolygons( PolygonsToCreate, CreatedPolygonRefs, CreatedEdgeIDs );
+	NewPolygonRefs.Append( CreatedPolygonRefs );
+
+	const bool bDeleteOrphanedEdges = true;
+	const bool bDeleteOrphanedVertices = false;
+	const bool bDeleteEmptySections = false;
+	DeletePolygons( PolygonRefsToDelete, bDeleteOrphanedEdges, bDeleteOrphanedVertices, bDeleteEmptySections );
+
+	GenerateNormalsAndTangentsForPolygonsAndAdjacents( NewPolygonRefs );
+}
+
+
+bool UEditableMesh::AnyChangesToUndo() const
+{
+	return bAllowUndo && Undo.IsValid() && Undo->Subchanges.Num() > 0;
+}
+
+
+void UEditableMesh::AddUndo( TUniquePtr<FChange> NewUndo )
+{
+	if( NewUndo.IsValid() )
+	{
+		if( bAllowUndo )
+		{
+			if( !Undo.IsValid() )
+			{
+				Undo = MakeUnique<FCompoundChangeInput>();
+			}
+
+			Undo->Subchanges.Add( MoveTemp( NewUndo ) );
+		}
+	}
+}
+
+
+TUniquePtr<FChange> UEditableMesh::MakeUndo()
+{
+	TUniquePtr<FChange> UndoChange = nullptr;
+	if( AnyChangesToUndo() )
+	{
+		UndoChange = MakeUnique<FCompoundChange>( MoveTemp( *Undo ) );
+	}
+	Undo.Reset();
+
+	return UndoChange;
+}
\ No newline at end of file
