Branch snapshot for CL 3327560

diff --git a/Engine/Source/Runtime/MeshEditingRuntime/EditableMeshChanges.cpp b/Engine/Source/Runtime/MeshEditingRuntime/EditableMeshChanges.cpp
new file mode 100644
--- /dev/null
+++ b/Engine/Source/Runtime/MeshEditingRuntime/EditableMeshChanges.cpp
@@ -0,0 +1,344 @@
+// Copyright 1998-2015 Epic Games, Inc. All Rights Reserved.
+
+#include "EditableMeshChanges.h"
+#include "EditableMesh.h"
+
+
+TUniquePtr<FChange> FDeleteOrphanVerticesChange::Execute( UObject* Object )
+{
+	UEditableMesh* EditableMesh = CastChecked<UEditableMesh>( Object );
+	verify( !EditableMesh->AnyChangesToUndo() );
+
+	EditableMesh->DeleteOrphanVertices( Input.VertexIDsToDelete );
+
+	return EditableMesh->MakeUndo();
+}
+
+
+FString FDeleteOrphanVerticesChange::ToString() const
+{
+	return FString::Printf(
+		TEXT( "Delete Orphan Vertices [VertexIDsToDelete:%s]" ),
+		*LogHelpers::ArrayToString( Input.VertexIDsToDelete ) );
+}
+
+
+TUniquePtr<FChange> FDeleteEdgesChange::Execute( UObject* Object )
+{
+	UEditableMesh* EditableMesh = CastChecked<UEditableMesh>( Object );
+	verify( !EditableMesh->AnyChangesToUndo() );
+
+	EditableMesh->DeleteEdges( Input.EdgeIDsToDelete, Input.bDeleteOrphanedVertices );
+
+	return EditableMesh->MakeUndo();
+}
+
+
+FString FDeleteEdgesChange::ToString() const
+{
+	return FString::Printf(
+		TEXT( "Delete Edges [EdgeIDsToDelete:%s, bDeleteOrphanedVertices:%s]" ),
+		*LogHelpers::ArrayToString( Input.EdgeIDsToDelete ),
+		*LogHelpers::BoolToString( Input.bDeleteOrphanedVertices ) );
+}
+
+
+TUniquePtr<FChange> FCreateVerticesChange::Execute( UObject* Object )
+{
+	UEditableMesh* EditableMesh = CastChecked<UEditableMesh>( Object );
+	verify( !EditableMesh->AnyChangesToUndo() );
+
+	static TArray<FVertexID> UnusedNewVertexIDs;
+	EditableMesh->CreateVertices( Input.VerticesToCreate, /* Out */ UnusedNewVertexIDs );
+
+	return EditableMesh->MakeUndo();
+}
+
+
+FString FCreateVerticesChange::ToString() const
+{
+	return FString::Printf(
+		TEXT( "Create Vertices [VerticesToCreate:%s]" ),
+		*LogHelpers::ArrayToString( Input.VerticesToCreate ) );
+}
+
+
+TUniquePtr<FChange> FCreateEdgesChange::Execute( UObject* Object )
+{
+	UEditableMesh* EditableMesh = CastChecked<UEditableMesh>( Object );
+	verify( !EditableMesh->AnyChangesToUndo() );
+
+	static TArray<FEdgeID> UnusedNewEdgeIDs;
+	EditableMesh->CreateEdges( Input.EdgesToCreate, /* Out */ UnusedNewEdgeIDs );
+
+	return EditableMesh->MakeUndo();
+}
+
+
+FString FCreateEdgesChange::ToString() const
+{
+	return FString::Printf(
+		TEXT( "Create Edges [EdgesToCreate:%s]" ),
+		*LogHelpers::ArrayToString( Input.EdgesToCreate ) );
+}
+
+
+TUniquePtr<FChange> FCreatePolygonsChange::Execute( UObject* Object )
+{
+	UEditableMesh* EditableMesh = CastChecked<UEditableMesh>( Object );
+	verify( !EditableMesh->AnyChangesToUndo() );
+	
+	static TArray<FPolygonRef> UnusedNewPolygonRefs;
+	static TArray<FEdgeID> UnusedNewEdgeIDs;
+	EditableMesh->CreatePolygons( Input.PolygonsToCreate, /* Out */ UnusedNewPolygonRefs, /* Out */ UnusedNewEdgeIDs );
+
+	return EditableMesh->MakeUndo();
+}
+
+
+FString FCreatePolygonsChange::ToString() const
+{
+	return FString::Printf(
+		TEXT( "Create Polygons [PolygonsToCreate:%s]" ),
+		*LogHelpers::ArrayToString( Input.PolygonsToCreate ) );
+}
+
+
+TUniquePtr<FChange> FDeletePolygonsChange::Execute( UObject* Object )
+{
+	UEditableMesh* EditableMesh = CastChecked<UEditableMesh>( Object );
+	verify( !EditableMesh->AnyChangesToUndo() );
+
+	EditableMesh->DeletePolygons( Input.PolygonRefsToDelete, Input.bDeleteOrphanedEdges, Input.bDeleteOrphanedVertices, Input.bDeleteEmptySections );
+	return EditableMesh->MakeUndo();
+}
+
+
+FString FDeletePolygonsChange::ToString() const
+{
+	return FString::Printf(
+		TEXT( "Delete Polygons [PolygonRefsToDelete:%s, %s, %s]" ),
+		*LogHelpers::ArrayToString( Input.PolygonRefsToDelete ),
+		*LogHelpers::BoolToString( Input.bDeleteOrphanedEdges ),
+		*LogHelpers::BoolToString( Input.bDeleteOrphanedVertices ),
+		*LogHelpers::BoolToString( Input.bDeleteEmptySections ) );
+}
+
+
+TUniquePtr<FChange> FSetVerticesAttributesChange::Execute( UObject* Object )
+{
+	UEditableMesh* EditableMesh = CastChecked<UEditableMesh>( Object );
+	verify( !EditableMesh->AnyChangesToUndo() );
+
+	EditableMesh->SetVerticesAttributes( Input.AttributesForVertices );
+	return EditableMesh->MakeUndo();
+}
+
+
+FString FSetVerticesAttributesChange::ToString() const
+{
+	return FString::Printf(
+		TEXT( "Set Vertices Attributes [AttributesForVertices:%s]" ),
+		*LogHelpers::ArrayToString( Input.AttributesForVertices ) );
+}
+
+
+TUniquePtr<FChange> FSetEdgesAttributesChange::Execute( UObject* Object )
+{
+	UEditableMesh* EditableMesh = CastChecked<UEditableMesh>( Object );
+	verify( !EditableMesh->AnyChangesToUndo() );
+
+	EditableMesh->SetEdgesAttributes( Input.AttributesForEdges );
+	return EditableMesh->MakeUndo();
+}
+
+
+FString FSetEdgesAttributesChange::ToString() const
+{
+	return FString::Printf(
+		TEXT( "Set Edges Attributes [AttributesForEdges:%s]" ),
+		*LogHelpers::ArrayToString( Input.AttributesForEdges ) );
+}
+
+
+TUniquePtr<FChange> FSetPolygonsVertexAttributesChange::Execute( UObject* Object )
+{
+	UEditableMesh* EditableMesh = CastChecked<UEditableMesh>( Object );
+	verify( !EditableMesh->AnyChangesToUndo() );
+
+	EditableMesh->SetPolygonsVertexAttributes( Input.VertexAttributesForPolygons );
+	return EditableMesh->MakeUndo();
+}
+
+
+FString FSetPolygonsVertexAttributesChange::ToString() const
+{
+	return FString::Printf(
+		TEXT( "Set Polygons Vertex Attributes [VertexAttributesForPolygons:%s]" ),
+		*LogHelpers::ArrayToString( Input.VertexAttributesForPolygons ) );
+}
+
+
+TUniquePtr<FChange> FSetEdgesVerticesChange::Execute( UObject* Object )
+{
+	UEditableMesh* EditableMesh = CastChecked<UEditableMesh>( Object );
+	verify( !EditableMesh->AnyChangesToUndo() );
+
+	EditableMesh->SetEdgesVertices( Input.VerticesForEdges );
+	return EditableMesh->MakeUndo();
+}
+
+
+
+FString FSetEdgesVerticesChange::ToString() const
+{
+	return FString::Printf(
+		TEXT( "Set Edges Vertices [VerticesForEdges:%s]" ),
+		*LogHelpers::ArrayToString( Input.VerticesForEdges ) );
+}
+
+
+TUniquePtr<FChange> FInsertPolygonPerimeterVerticesChange::Execute( UObject* Object )
+{
+	UEditableMesh* EditableMesh = CastChecked<UEditableMesh>( Object );
+	verify( !EditableMesh->AnyChangesToUndo() );
+
+	EditableMesh->InsertPolygonPerimeterVertices( Input.PolygonRef, Input.InsertBeforeVertexNumber, Input.VerticesToInsert );
+	return EditableMesh->MakeUndo();
+}
+
+
+FString FInsertPolygonPerimeterVerticesChange::ToString() const
+{
+	return FString::Printf(
+		TEXT( "Insert Polygon Perimeter Vertices [PolygonRef:%s, InsertBeforeVertexNumber:%lu, VerticesToInsert:%s]" ),
+		*Input.PolygonRef.ToString(),
+		Input.InsertBeforeVertexNumber,
+		*LogHelpers::ArrayToString( Input.VerticesToInsert ) );
+}
+
+
+TUniquePtr<FChange> FRemovePolygonPerimeterVerticesChange::Execute( UObject* Object )
+{
+	UEditableMesh* EditableMesh = CastChecked<UEditableMesh>( Object );
+	verify( !EditableMesh->AnyChangesToUndo() );
+
+	EditableMesh->RemovePolygonPerimeterVertices( Input.PolygonRef, Input.FirstVertexNumberToRemove, Input.NumVerticesToRemove );
+	return EditableMesh->MakeUndo();
+}
+
+
+FString FRemovePolygonPerimeterVerticesChange::ToString() const
+{
+	return FString::Printf(
+		TEXT( "Remove Polygon Perimeter Vertices [PolygonRef:%s, FirstVertexNumberToRemove:%lu, NumVerticesToRemove:%lu]" ),
+		*Input.PolygonRef.ToString(),
+		Input.FirstVertexNumberToRemove,
+		Input.NumVerticesToRemove );
+}
+
+
+TUniquePtr<FChange> FStartOrEndModificationChange::Execute( UObject* Object )
+{
+	UEditableMesh* EditableMesh = CastChecked<UEditableMesh>( Object );
+	verify( !EditableMesh->AnyChangesToUndo() );
+
+	if( Input.bStartModification )
+	{
+		EditableMesh->StartModification( Input.MeshModificationType, Input.MeshTopologyChange );
+	}
+	else
+	{
+		const bool bFromUndo = true;
+		EditableMesh->EndModification( bFromUndo );
+	}
+
+	return EditableMesh->MakeUndo();
+}
+
+
+FString FStartOrEndModificationChange::ToString() const
+{
+	return FString::Printf(
+		TEXT( "%s Modification (MeshModificationType:%i, MeshTopologyChange:%i)" ),
+		Input.bStartModification ? TEXT( "Start" ) : TEXT( "End" ),
+		(int32)Input.MeshModificationType,
+		(int32)Input.MeshTopologyChange );
+}
+
+
+TUniquePtr<FChange> FRetrianglulatePolygonsChange::Execute( UObject* Object )
+{
+	UEditableMesh* EditableMesh = CastChecked<UEditableMesh>( Object );
+	verify( !EditableMesh->AnyChangesToUndo() );
+
+	EditableMesh->RetriangulatePolygons( Input.PolygonRefs, Input.bOnlyOnUndo );
+	return EditableMesh->MakeUndo();
+}
+
+
+FString FRetrianglulatePolygonsChange::ToString() const
+{
+	return FString::Printf(
+		TEXT( "Retriangulate Polygons [PolygonRefs:%s, bOnlyOnUndo:%s]" ),
+		*LogHelpers::ArrayToString( Input.PolygonRefs ),
+		*LogHelpers::BoolToString( Input.bOnlyOnUndo ) );
+}
+
+
+TUniquePtr<FChange> FSetSubdivisionCountChange::Execute( UObject* Object )
+{
+	UEditableMesh* EditableMesh = CastChecked<UEditableMesh>( Object );
+	verify( !EditableMesh->AnyChangesToUndo() );
+
+	EditableMesh->SetSubdivisionCount( Input.NewSubdivisionCount );
+	return EditableMesh->MakeUndo();
+}
+
+
+FString FSetSubdivisionCountChange::ToString() const
+{
+	return FString::Printf(
+		TEXT( "Set Subdivision Count [NewSubdivisionCount:%i]" ),
+		Input.NewSubdivisionCount );
+}
+
+
+TUniquePtr<FChange> FCreateSectionChange::Execute( UObject* Object )
+{
+	UEditableMesh* EditableMesh = CastChecked<UEditableMesh>( Object );
+	verify( !EditableMesh->AnyChangesToUndo() );
+
+	EditableMesh->CreateSection( Input.SectionToCreate );
+	return EditableMesh->MakeUndo();
+}
+
+
+FString FCreateSectionChange::ToString() const
+{
+	return FString::Printf(
+		TEXT( "Create Section [Material:%s, bEnableCollision:%i, bCastShadow:%i]" ),
+		Input.SectionToCreate.Material ? *Input.SectionToCreate.Material->GetName() : TEXT( "<none>" ),
+		Input.SectionToCreate.bEnableCollision,
+		Input.SectionToCreate.bCastShadow );
+}
+
+
+TUniquePtr<FChange> FDeleteSectionChange::Execute( UObject* Object )
+{
+	UEditableMesh* EditableMesh = CastChecked<UEditableMesh>( Object );
+	verify( !EditableMesh->AnyChangesToUndo() );
+
+	EditableMesh->DeleteSection( Input.SectionID );
+	return EditableMesh->MakeUndo();
+}
+
+
+FString FDeleteSectionChange::ToString() const
+{
+	return FString::Printf(
+		TEXT( "Delete Section [SectionID:%i]" ),
+		Input.SectionID.GetValue() );
+}
+
+
diff --git a/Engine/Source/Runtime/MeshEditingRuntime/EditableMeshChanges.h b/Engine/Source/Runtime/MeshEditingRuntime/EditableMeshChanges.h
new file mode 100644
--- /dev/null
+++ b/Engine/Source/Runtime/MeshEditingRuntime/EditableMeshChanges.h
@@ -0,0 +1,735 @@
+// Copyright 1998-2015 Epic Games, Inc. All Rights Reserved.
+
+#pragma once
+
+#include "CoreMinimal.h"
+#include "Change.h"
+#include "EditableMeshTypes.h"
+
+
+// @todo mesheditor: Comment these classes and enums!
+
+struct FDeleteOrphanVerticesChangeInput
+{
+	/** The vertex IDs to get rid of */
+	TArray<FVertexID> VertexIDsToDelete;
+
+	/** Default constructor */
+	FDeleteOrphanVerticesChangeInput()
+		: VertexIDsToDelete()
+	{
+	}
+};
+
+
+class FDeleteOrphanVerticesChange : public FChange
+{
+
+public:
+
+	/** Constructor */
+	FDeleteOrphanVerticesChange( const FDeleteOrphanVerticesChangeInput& InitInput )
+		: Input( InitInput )
+	{
+	}
+
+ 	FDeleteOrphanVerticesChange( FDeleteOrphanVerticesChangeInput&& InitInput )
+ 		: Input( MoveTemp( InitInput ) )
+ 	{
+ 	}
+
+	// Parent class overrides
+	virtual TUniquePtr<FChange> Execute( UObject* Object ) override;
+	virtual FString ToString() const override;
+
+private:
+	
+	/** The data we need to make this change */
+	FDeleteOrphanVerticesChangeInput Input;
+
+};
+
+
+struct FDeleteEdgesChangeInput
+{
+	/** The edge IDs to get rid of */
+	TArray<FEdgeID> EdgeIDsToDelete;
+
+	/** Whether we should also delete any vertices that are left orphaned after deleting this edge */
+	bool bDeleteOrphanedVertices;
+
+	/** Default constructor */
+	FDeleteEdgesChangeInput()
+		: EdgeIDsToDelete(),
+		  bDeleteOrphanedVertices( true )
+	{
+	}
+};
+
+
+class FDeleteEdgesChange : public FChange
+{
+
+public:
+
+	/** Constructor */
+	FDeleteEdgesChange( const FDeleteEdgesChangeInput& InitInput )
+		: Input( InitInput )
+	{
+	}
+
+ 	FDeleteEdgesChange( FDeleteEdgesChangeInput&& InitInput )
+ 		: Input( MoveTemp( InitInput ) )
+ 	{
+ 	}
+
+	// Parent class overrides
+	virtual TUniquePtr<FChange> Execute( UObject* Object ) override;
+	virtual FString ToString() const override;
+
+
+private:
+
+	/** The data we need to make this change */
+	FDeleteEdgesChangeInput Input;
+
+};
+
+
+struct FCreateVerticesChangeInput
+{
+	/** Information about each vertex that will be created */
+	TArray<FVertexToCreate> VerticesToCreate;
+
+	/** Default constructor */
+	FCreateVerticesChangeInput()
+		: VerticesToCreate()
+	{
+	}
+};
+
+
+class FCreateVerticesChange : public FChange
+{
+
+public:
+
+	/** Constructor */
+	FCreateVerticesChange( const FCreateVerticesChangeInput& InitInput )
+		: Input( InitInput )
+	{
+	}
+
+ 	FCreateVerticesChange( FCreateVerticesChangeInput&& InitInput )
+ 		: Input( MoveTemp( InitInput ) )
+ 	{
+ 	}
+
+	// Parent class overrides
+	virtual TUniquePtr<FChange> Execute( UObject* Object ) override;
+	virtual FString ToString() const override;
+
+
+private:
+	
+	/** The data we need to make this change */
+	FCreateVerticesChangeInput Input;
+
+};
+
+
+struct FCreateEdgesChangeInput
+{
+	/** Information about each edge that will be created */
+	TArray<FEdgeToCreate> EdgesToCreate;
+
+	/** Default constructor */
+	FCreateEdgesChangeInput()
+		: EdgesToCreate()
+	{
+	}
+};
+
+
+class FCreateEdgesChange : public FChange
+{
+
+public:
+
+	/** Constructor */
+	FCreateEdgesChange( const FCreateEdgesChangeInput& InitInput )
+		: Input( InitInput )
+	{
+	}
+
+ 	FCreateEdgesChange( FCreateEdgesChangeInput&& InitInput )
+ 		: Input( MoveTemp( InitInput ) )
+ 	{
+ 	}
+
+	// Parent class overrides
+	virtual TUniquePtr<FChange> Execute( UObject* Object ) override;
+	virtual FString ToString() const override;
+
+
+private:
+
+	/** The data we need to make this change */
+	FCreateEdgesChangeInput Input;
+
+};
+
+
+struct FCreatePolygonsChangeInput
+{
+	/** Information about each polygon that will be created */
+	TArray<FPolygonToCreate> PolygonsToCreate;
+
+	/** Default constructor */
+	FCreatePolygonsChangeInput()
+		: PolygonsToCreate()
+	{
+	}
+};
+
+
+class FCreatePolygonsChange : public FChange
+{
+
+public:
+
+	/** Constructor */
+	FCreatePolygonsChange( const FCreatePolygonsChangeInput& InitInput )
+		: Input( InitInput )
+	{
+	}
+
+ 	FCreatePolygonsChange( FCreatePolygonsChangeInput&& InitInput )
+ 		: Input( MoveTemp( InitInput ) )
+ 	{
+ 	}
+
+	// Parent class overrides
+	virtual TUniquePtr<FChange> Execute( UObject* Object ) override;
+	virtual FString ToString() const override;
+
+
+private:
+
+	/** The data we need to make this change */
+	FCreatePolygonsChangeInput Input;
+
+};
+
+
+struct FDeletePolygonsChangeInput
+{
+	/** The polygons to get rid of */
+	TArray<FPolygonRef> PolygonRefsToDelete;
+
+	/** Whether we should also delete any edges that are left orphaned after deleting this polygon */
+	bool bDeleteOrphanedEdges;
+
+	/** Whether we should also delete any vertices that are left orphaned after deleting this polygon */
+	bool bDeleteOrphanedVertices;
+
+	/** Whether we should also delete any sections that are left empty after deleting this polygon */
+	bool bDeleteEmptySections;
+
+	/** Default constructor */
+	FDeletePolygonsChangeInput()
+		: PolygonRefsToDelete(),
+		  bDeleteOrphanedEdges( true ),
+		  bDeleteOrphanedVertices( true ),
+		  bDeleteEmptySections( false )
+	{
+	}
+};
+
+
+class FDeletePolygonsChange : public FChange
+{
+
+public:
+
+	/** Constructor */
+	FDeletePolygonsChange( const FDeletePolygonsChangeInput& InitInput )
+		: Input( InitInput )
+	{
+	}
+
+ 	FDeletePolygonsChange( FDeletePolygonsChangeInput&& InitInput )
+ 		: Input( MoveTemp ( InitInput ) )
+ 	{
+ 	}
+
+	// Parent class overrides
+	virtual TUniquePtr<FChange> Execute( UObject* Object ) override;
+	virtual FString ToString() const override;
+
+
+private:
+
+	/** The data we need to make this change */
+	FDeletePolygonsChangeInput Input;
+
+};
+
+
+struct FSetVerticesAttributesChangeInput
+{
+	TArray<FAttributesForVertex> AttributesForVertices;
+};
+
+
+class FSetVerticesAttributesChange : public FChange
+{
+
+public:
+
+	/** Constructor */
+	FSetVerticesAttributesChange( const FSetVerticesAttributesChangeInput& InitInput )
+		: Input( InitInput )
+	{
+	}
+
+ 	FSetVerticesAttributesChange( FSetVerticesAttributesChangeInput&& InitInput )
+ 		: Input( MoveTemp( InitInput ) )
+ 	{
+ 	}
+
+	// Parent class overrides
+	virtual TUniquePtr<FChange> Execute( UObject* Object ) override;
+	virtual FString ToString() const override;
+
+
+private:
+
+	/** The data we need to make this change */
+	FSetVerticesAttributesChangeInput Input;
+
+};
+
+
+struct FSetEdgesAttributesChangeInput
+{
+	/** Which edges we'll be setting attributes for, along with the attribute data to set */
+	TArray<FAttributesForEdge> AttributesForEdges;
+
+	/** Default constructor */
+	FSetEdgesAttributesChangeInput()
+		: AttributesForEdges()
+	{
+	}
+};
+
+
+class FSetEdgesAttributesChange : public FChange
+{
+
+public:
+
+	/** Constructor */
+	FSetEdgesAttributesChange( const FSetEdgesAttributesChangeInput& InitInput )
+		: Input( InitInput )
+	{
+	}
+
+ 	FSetEdgesAttributesChange( FSetEdgesAttributesChangeInput&& InitInput )
+ 		: Input( MoveTemp( InitInput ) )
+ 	{
+ 	}
+
+	// Parent class overrides
+	virtual TUniquePtr<FChange> Execute( UObject* Object ) override;
+	virtual FString ToString() const override;
+
+
+private:
+
+	/** The data we need to make this change */
+	FSetEdgesAttributesChangeInput Input;
+
+};
+
+
+struct FSetPolygonsVertexAttributesChangeInput
+{
+	/** Which polygons we'll be setting vertex attributes for, along with the attribute data to set */
+	TArray<FVertexAttributesForPolygon> VertexAttributesForPolygons;
+
+	/** Default constructor */
+	FSetPolygonsVertexAttributesChangeInput()
+		: VertexAttributesForPolygons()
+	{
+	}
+};
+
+
+class FSetPolygonsVertexAttributesChange : public FChange
+{
+
+public:
+
+	/** Constructor */
+	FSetPolygonsVertexAttributesChange( const FSetPolygonsVertexAttributesChangeInput& InitInput )
+		: Input( InitInput )
+	{
+	}
+
+ 	FSetPolygonsVertexAttributesChange( FSetPolygonsVertexAttributesChangeInput&& InitInput )
+ 		: Input( MoveTemp( InitInput ) )
+ 	{
+ 	}
+
+	// Parent class overrides
+	virtual TUniquePtr<FChange> Execute( UObject* Object ) override;
+	virtual FString ToString() const override;
+
+
+private:
+
+	/** The data we need to make this change */
+	FSetPolygonsVertexAttributesChangeInput Input;
+
+};
+
+
+struct FSetEdgesVerticesChangeInput
+{
+	/** The edge to set new vertices for */
+	TArray<FVerticesForEdge> VerticesForEdges;
+
+	/** Default constructor */
+	FSetEdgesVerticesChangeInput()
+		: VerticesForEdges()
+	{
+	}
+};
+
+
+class FSetEdgesVerticesChange : public FChange
+{
+
+public:
+
+	/** Constructor */
+	FSetEdgesVerticesChange( const FSetEdgesVerticesChangeInput& InitInput )
+		: Input( InitInput )
+	{
+	}
+
+ 	FSetEdgesVerticesChange( FSetEdgesVerticesChangeInput&& InitInput )
+ 		: Input( MoveTemp( InitInput ) )
+ 	{
+ 	}
+
+	// Parent class overrides
+	virtual TUniquePtr<FChange> Execute( UObject* Object ) override;
+	virtual FString ToString() const override;
+
+
+private:
+
+	/** The data we need to make this change */
+	FSetEdgesVerticesChangeInput Input;
+
+};
+
+
+struct FInsertPolygonPerimeterVerticesChangeInput
+{
+	/** The polygon we'll be inserting vertices into */
+	FPolygonRef PolygonRef;
+
+	/** The first polygon perimeter vertex number to insert indices before */
+	uint32 InsertBeforeVertexNumber;
+
+	/** The vertices to insert, along with their polygon perimeter vertex attributes */
+	TArray<FVertexAndAttributes> VerticesToInsert;
+
+	
+	/** Default constructor */
+	FInsertPolygonPerimeterVerticesChangeInput()
+		: PolygonRef( FSectionID::Invalid, FPolygonID::Invalid ),
+		  InsertBeforeVertexNumber( 0 ),
+		  VerticesToInsert()
+	{
+	}
+};
+
+
+class FInsertPolygonPerimeterVerticesChange : public FChange
+{
+
+public:
+
+	/** Constructor */
+	FInsertPolygonPerimeterVerticesChange( const FInsertPolygonPerimeterVerticesChangeInput& InitInput )
+		: Input( InitInput )
+	{
+	}
+
+ 	FInsertPolygonPerimeterVerticesChange( FInsertPolygonPerimeterVerticesChangeInput&& InitInput )
+ 		: Input( MoveTemp( InitInput ) )
+ 	{
+ 	}
+
+	// Parent class overrides
+	virtual TUniquePtr<FChange> Execute( UObject* Object ) override;
+	virtual FString ToString() const override;
+
+
+private:
+
+	/** The data we need to make this change */
+	FInsertPolygonPerimeterVerticesChangeInput Input;
+
+};
+
+
+struct FRemovePolygonPerimeterVerticesChangeInput
+{
+	/** The polygon we'll be removing vertices from */
+	FPolygonRef PolygonRef;
+
+	/** The first polygon perimeter vertex number to remove */
+	uint32 FirstVertexNumberToRemove;
+
+	/** The number of vertices to remove */
+	uint32 NumVerticesToRemove;
+
+	
+	/** Default constructor */
+	FRemovePolygonPerimeterVerticesChangeInput()
+		: PolygonRef( FSectionID::Invalid, FPolygonID::Invalid ),
+		  FirstVertexNumberToRemove( 0 ),
+		  NumVerticesToRemove( 0 )
+	{
+	}
+};
+
+
+class FRemovePolygonPerimeterVerticesChange : public FChange
+{
+
+public:
+
+	/** Constructor */
+	FRemovePolygonPerimeterVerticesChange( const FRemovePolygonPerimeterVerticesChangeInput& InitInput )
+		: Input( InitInput )
+	{
+	}
+
+	// Parent class overrides
+	virtual TUniquePtr<FChange> Execute( UObject* Object ) override;
+	virtual FString ToString() const override;
+
+
+private:
+
+	/** The data we need to make this change */
+	FRemovePolygonPerimeterVerticesChangeInput Input;
+
+};
+
+
+struct FStartOrEndModificationChangeInput
+{
+	/** True if we should start modifying the mesh, or false if we should end modifying the mesh.  This will be reversed every time an undo/redo happens */
+	bool bStartModification;
+
+	/** The type of modification we're doing here */
+	EMeshModificationType MeshModificationType;
+	
+	/** Whether the mesh's topology can change during this modification */
+	EMeshTopologyChange MeshTopologyChange;
+
+	
+	/** Default constructor */
+	FStartOrEndModificationChangeInput()
+		: bStartModification( true ),
+		  MeshModificationType( EMeshModificationType::Final ),
+		  MeshTopologyChange( EMeshTopologyChange::TopologyChange )
+	{
+	}
+};
+
+
+class FStartOrEndModificationChange : public FChange
+{
+
+public:
+
+	/** Constructor */
+	FStartOrEndModificationChange( const FStartOrEndModificationChangeInput& InitInput )
+		: Input( InitInput )
+	{
+	}
+
+	// Parent class overrides
+	virtual TUniquePtr<FChange> Execute( UObject* Object ) override;
+	virtual FString ToString() const override;
+
+
+private:
+
+	/** The data we need to make this change */
+	FStartOrEndModificationChangeInput Input;
+
+};
+
+
+struct FRetrianglulatePolygonsChangeInput
+{
+	/** The polygons to retriangulate */
+	TArray<FPolygonRef> PolygonRefs;
+
+	/** Whether we should actually do the retriangulation.  This will be inverted every time an Undo/Redo happens */
+	bool bOnlyOnUndo;
+
+	/** Default constructor */
+	FRetrianglulatePolygonsChangeInput()
+		: PolygonRefs(),
+		  bOnlyOnUndo( true )
+	{
+	}
+};
+
+
+class FRetrianglulatePolygonsChange : public FChange
+{
+
+public:
+
+	/** Constructor */
+	FRetrianglulatePolygonsChange( const FRetrianglulatePolygonsChangeInput& InitInput )
+		: Input( InitInput )
+	{
+	}
+
+ 	FRetrianglulatePolygonsChange( FRetrianglulatePolygonsChangeInput&& InitInput )
+ 		: Input( MoveTemp( InitInput ) )
+ 	{
+ 	}
+
+	// Parent class overrides
+	virtual TUniquePtr<FChange> Execute( UObject* Object ) override;
+	virtual FString ToString() const override;
+
+
+private:
+
+	/** The data we need to make this change */
+	FRetrianglulatePolygonsChangeInput Input;
+
+};
+
+
+struct FSetSubdivisionCountChangeInput
+{
+	/** The new number of subdivisions to use */
+	int32 NewSubdivisionCount;
+
+	/** Default constructor */
+	FSetSubdivisionCountChangeInput()
+		: NewSubdivisionCount( 0 )
+	{
+	}
+};
+
+
+class FSetSubdivisionCountChange : public FChange
+{
+
+public:
+
+	/** Constructor */
+	FSetSubdivisionCountChange( const FSetSubdivisionCountChangeInput& InitInput )
+		: Input( InitInput )
+	{
+	}
+
+	// Parent class overrides
+	virtual TUniquePtr<FChange> Execute( UObject* Object ) override;
+	virtual FString ToString() const override;
+
+
+private:
+
+	/** The data we need to make this change */
+	FSetSubdivisionCountChangeInput Input;
+
+};
+
+
+struct FCreateSectionChangeInput
+{
+	/** Information about the section to create */
+	FSectionToCreate SectionToCreate;
+
+	/** Default constructor */
+	FCreateSectionChangeInput()
+		: SectionToCreate()
+	{
+	}
+};
+
+
+class FCreateSectionChange : public FChange
+{
+
+public:
+
+	/** Constructor */
+	FCreateSectionChange( const FCreateSectionChangeInput& InitInput )
+		: Input( InitInput )
+	{
+	}
+
+	// Parent class overrides
+	virtual TUniquePtr<FChange> Execute( UObject* Object ) override;
+	virtual FString ToString() const override;
+
+
+private:
+
+	/** The data we need to make this change */
+	FCreateSectionChangeInput Input;
+};
+
+
+struct FDeleteSectionChangeInput
+{
+	/** The section to delete */
+	FSectionID SectionID;
+
+	/** Default constructor */
+	FDeleteSectionChangeInput()
+		: SectionID( FSectionID::Invalid )
+	{
+	}
+};
+
+
+class FDeleteSectionChange : public FChange
+{
+
+public:
+
+	/** Constructor */
+	FDeleteSectionChange( const FDeleteSectionChangeInput& InitInput )
+		: Input( InitInput )
+	{
+	}
+
+	// Parent class overrides
+	virtual TUniquePtr<FChange> Execute( UObject* Object ) override;
+	virtual FString ToString() const override;
+
+
+private:
+
+	/** The data we need to make this change */
+	FDeleteSectionChangeInput Input;
+};
diff --git a/Engine/Source/Runtime/MeshEditingRuntime/EditableMeshCustomVersion.cpp b/Engine/Source/Runtime/MeshEditingRuntime/EditableMeshCustomVersion.cpp
new file mode 100644
--- /dev/null
+++ b/Engine/Source/Runtime/MeshEditingRuntime/EditableMeshCustomVersion.cpp
@@ -0,0 +1,9 @@
+// Copyright 1998-2016 Epic Games, Inc. All Rights Reserved.
+
+#include "EditableMeshCustomVersion.h"
+#include "Serialization/CustomVersion.h"
+
+const FGuid FEditableMeshCustomVersion::GUID( 0xFB26E412, 0x1F154B4D, 0x9372550A, 0x961D2F70 );
+
+// Register the custom version with core
+FCustomVersionRegistration GRegisterEditableMeshCustomVersion( FEditableMeshCustomVersion::GUID, FEditableMeshCustomVersion::LatestVersion, TEXT( "EditableMeshVer" ) );
diff --git a/Engine/Source/Runtime/MeshEditingRuntime/EditableMeshCustomVersion.h b/Engine/Source/Runtime/MeshEditingRuntime/EditableMeshCustomVersion.h
new file mode 100644
--- /dev/null
+++ b/Engine/Source/Runtime/MeshEditingRuntime/EditableMeshCustomVersion.h
@@ -0,0 +1,34 @@
+// Copyright 1998-2016 Epic Games, Inc. All Rights Reserved.
+
+#pragma once
+
+#include "CoreMinimal.h"
+#include "Misc/Guid.h"
+
+// Custom serialization version for all packages containing EditableMesh asset types
+struct FEditableMeshCustomVersion
+{
+	enum Type
+	{
+		// Before any version changes were made in the module
+		BeforeCustomVersionWasAdded = 0,
+
+		// Added subdivision level count and texture coordinate count to UEditableMesh format
+		TextureCoordinateAndSubdivisionCounts,
+
+		// SparseArrays now have custom serialization, preserving the gaps
+		CustomSparseArraySerialization,
+
+		// Added vertex color support
+		WithVertexColors,
+
+		// -----<new versions can be added above this line>-------------------------------------------------
+		VersionPlusOne,
+		LatestVersion = VersionPlusOne - 1
+	};
+
+	// The GUID for this custom version number
+	const static FGuid GUID;
+
+	FEditableMeshCustomVersion() = delete;
+};
diff --git a/Engine/Source/Runtime/MeshEditingRuntime/EditableMeshFactory.cpp b/Engine/Source/Runtime/MeshEditingRuntime/EditableMeshFactory.cpp
new file mode 100644
--- /dev/null
+++ b/Engine/Source/Runtime/MeshEditingRuntime/EditableMeshFactory.cpp
@@ -0,0 +1,81 @@
+// Copyright 1998-2015 Epic Games, Inc. All Rights Reserved.
+
+#include "EditableMeshFactory.h"
+#include "Features/IModularFeatures.h"
+#include "IEditableMeshFormat.h"
+
+
+FEditableMeshSubMeshAddress UEditableMeshFactory::MakeSubmeshAddress( UPrimitiveComponent* PrimitiveComponent, const int32 LODIndex )
+{
+	check( PrimitiveComponent != nullptr );
+
+	FEditableMeshSubMeshAddress SubMeshAddress;
+
+	const int32 NumEditableMeshFormats = IModularFeatures::Get().GetModularFeatureImplementationCount( "EditableMeshFormat" );
+	for( int32 EditableMeshFormatIndex = 0; EditableMeshFormatIndex < NumEditableMeshFormats; ++EditableMeshFormatIndex )
+	{
+		IEditableMeshFormat& EditableMeshFormat = *static_cast<IEditableMeshFormat*>( IModularFeatures::Get().GetModularFeatureImplementation( "EditableMeshFormat", EditableMeshFormatIndex ) );
+
+		SubMeshAddress = FEditableMeshSubMeshAddress();
+		SubMeshAddress.MeshObjectPtr = nullptr;	// This will be filled in below (FillMeshObjectPtr)
+		SubMeshAddress.EditableMeshFormat = &EditableMeshFormat;
+		SubMeshAddress.LODIndex = LODIndex;
+		EditableMeshFormat.FillMeshObjectPtr( *PrimitiveComponent, SubMeshAddress );	// @todo mesheditor: This stuff is a bit clunky, would like to refactor it
+		if( SubMeshAddress.MeshObjectPtr != nullptr )
+		{
+			break;
+		}
+		else
+		{
+			SubMeshAddress = FEditableMeshSubMeshAddress();
+		}
+	}
+
+	return SubMeshAddress;
+}
+
+
+FEditableMeshSubMeshAddress UEditableMeshFactory::MakeSubmeshAddress( UStaticMesh& StaticMesh, const int32 LODIndex )
+{
+	FEditableMeshSubMeshAddress SubMeshAddress;
+
+	const int32 NumEditableMeshFormats = IModularFeatures::Get().GetModularFeatureImplementationCount("EditableMeshFormat");
+	
+	if(NumEditableMeshFormats > 0)
+	{
+		IEditableMeshFormat& EditableMeshFormat = *static_cast<IEditableMeshFormat*>(IModularFeatures::Get().GetModularFeatureImplementation("EditableMeshFormat", 0));
+
+		SubMeshAddress = FEditableMeshSubMeshAddress();
+		SubMeshAddress.MeshObjectPtr = &StaticMesh;	
+		SubMeshAddress.EditableMeshFormat = &EditableMeshFormat;
+		SubMeshAddress.LODIndex = LODIndex;
+	}
+
+	return SubMeshAddress;
+}
+
+UEditableMesh* UEditableMeshFactory::MakeEditableMesh( UPrimitiveComponent* PrimitiveComponent, const int32 LODIndex )
+{
+	check( PrimitiveComponent != nullptr );
+
+	FEditableMeshSubMeshAddress SubMeshAddress = MakeSubmeshAddress( PrimitiveComponent, LODIndex );
+	return MakeEditableMesh( PrimitiveComponent, SubMeshAddress );
+}
+
+
+UEditableMesh* UEditableMeshFactory::MakeEditableMesh( UPrimitiveComponent* PrimitiveComponent, const FEditableMeshSubMeshAddress& SubMeshAddress )
+{
+	check( PrimitiveComponent != nullptr );
+	UEditableMesh* EditableMesh = nullptr;
+
+	if( SubMeshAddress.EditableMeshFormat != nullptr &&
+		SubMeshAddress.MeshObjectPtr != nullptr )
+	{
+		// @todo mesheditor perf: This is going to HITCH
+		EditableMesh = SubMeshAddress.EditableMeshFormat->MakeEditableMesh( *PrimitiveComponent, SubMeshAddress );
+	}
+
+	return EditableMesh;
+}
+
+
