Branch snapshot for CL 3327560

diff --git a/Engine/Source/Runtime/MeshEditingRuntime/EditableStaticMesh.cpp b/Engine/Source/Runtime/MeshEditingRuntime/EditableStaticMesh.cpp
new file mode 100644
--- /dev/null
+++ b/Engine/Source/Runtime/MeshEditingRuntime/EditableStaticMesh.cpp
@@ -0,0 +1,3314 @@
+// Copyright 1998-2015 Epic Games, Inc. All Rights Reserved.
+
+#include "EditableStaticMesh.h"
+#include "EditableMeshChanges.h"
+#include "Engine/StaticMesh.h"
+#include "Components/StaticMeshComponent.h"
+#include "StaticMeshResources.h"
+#include "PhysicsEngine/BodySetup.h"	// For collision generation
+#include "ProfilingDebugging/ScopedTimers.h"	// For FAutoScopedDurationTimer
+#include "EditableMeshFactory.h"
+
+
+const FRenderingVertexID FRenderingVertexID::Invalid( TNumericLimits<uint32>::Max() );
+const FTriangleID FTriangleID::Invalid( TNumericLimits<uint32>::Max() );
+
+
+UEditableStaticMesh::UEditableStaticMesh()
+	: UEditableMesh(),
+	  StaticMesh( nullptr ),
+	  Vertices(),
+	  Edges(),
+	  Sections(),
+	  RecreateRenderStateContext(),
+	  PendingCompactCounter( 0 )
+{
+}
+
+
+inline void UEditableStaticMesh::EnsureIndexBufferIs32Bit()
+{
+	FStaticMeshLODResources& StaticMeshLOD = GetStaticMeshLOD();
+	if( !StaticMeshLOD.IndexBuffer.Is32Bit() )
+	{
+		// Need a 32-bit index buffer
+		static TArray<uint32> AllIndices;
+		StaticMeshLOD.IndexBuffer.GetCopy( /* Out */ AllIndices );
+		StaticMeshLOD.IndexBuffer.SetIndices( AllIndices, EIndexBufferStride::Force32Bit );
+	}
+}
+
+
+inline void UEditableStaticMesh::UpdateIndexBufferFormatIfNeeded( const TArray<FRenderingVertexID>& RenderingVertexIDs )
+{
+	check( !IsPreviewingSubdivisions() );	// Should not be mutating the actual mesh when in subdivision preview mode
+
+	const FStaticMeshLODResources& StaticMeshLOD = GetStaticMeshLOD();
+	if( !StaticMeshLOD.IndexBuffer.Is32Bit() )
+	{
+		for( const FRenderingVertexID RenderingVertexID : RenderingVertexIDs )
+		{
+			if( RenderingVertexID.GetValue() > TNumericLimits<uint16>::Max() )
+			{
+				EnsureIndexBufferIs32Bit();
+				break;
+			}
+		}
+	}
+}
+
+
+inline void UEditableStaticMesh::UpdateIndexBufferFormatIfNeeded( const FRenderingVertexID RenderingVertexID )
+{
+	if( RenderingVertexID.GetValue() > TNumericLimits<uint16>::Max() )
+	{
+		EnsureIndexBufferIs32Bit();
+	}
+}
+
+
+void UEditableStaticMesh::InitEditableStaticMesh( class UPrimitiveComponent& Component, const FEditableMeshSubMeshAddress& InitSubMeshAddress )
+{
+	SetSubMeshAddress( InitSubMeshAddress );
+
+	// We're partial to static mesh components, here
+	UStaticMeshComponent* StaticMeshComponent = Cast<UStaticMeshComponent>( &Component );
+	if( StaticMeshComponent != nullptr )
+	{
+		UStaticMesh* ComponentStaticMesh = StaticMeshComponent->GetStaticMesh();
+		if( ComponentStaticMesh != nullptr && ComponentStaticMesh->HasValidRenderData() )
+		{
+			this->StaticMesh = ComponentStaticMesh;
+			this->OriginalStaticMesh = ComponentStaticMesh;
+
+			const FStaticMeshRenderData& StaticMeshRenderData = *StaticMesh->RenderData;
+			if( SubMeshAddress.LODIndex >= 0 && SubMeshAddress.LODIndex < StaticMeshRenderData.LODResources.Num() )
+			{
+				{
+					// @todo mesheditor urgent: Currently, we're disabling many of the optimized index buffers that were precomputed
+					// for static meshes when they become editable.  This is just so that we don't have to keep this data up to
+					// date as we perform live edits to the geometry.  Later, we should probably get this updated as we go, or 
+					// lazily update the buffers when committing a final change or saving.  Without clearing these values, some
+					// graphical artifacts will be visible while editing the mesh (flickering shadows, for example.)
+					FStaticMeshLODResources& StaticMeshLOD = StaticMesh->RenderData->LODResources[ SubMeshAddress.LODIndex ];
+					StaticMeshLOD.bHasAdjacencyInfo = false;
+					StaticMeshLOD.bHasDepthOnlyIndices = false;
+					StaticMeshLOD.bHasReversedIndices = false;
+					StaticMeshLOD.bHasReversedDepthOnlyIndices = false;
+					StaticMeshLOD.DepthOnlyNumTriangles = 0;
+				}
+
+				const FStaticMeshLODResources& StaticMeshLOD = StaticMeshRenderData.LODResources[ SubMeshAddress.LODIndex ];
+
+				// Store off the number of texture coordinates in this mesh
+				this->TextureCoordinateCount = StaticMeshLOD.GetNumTexCoords();
+
+				// Vertices
+				const int32 NumRenderingVertices = StaticMeshLOD.PositionVertexBuffer.GetNumVertices();
+				const int32 NumUVs = TextureCoordinateCount;
+				const bool bHasColor = StaticMeshLOD.ColorVertexBuffer.GetNumVertices() > 0;
+				check( !bHasColor || StaticMeshLOD.ColorVertexBuffer.GetNumVertices() == StaticMeshLOD.VertexBuffer.GetNumVertices() );
+				RenderingVertices.Reserve( NumRenderingVertices );
+
+				// @todo mesheditor cleanup: This code is very similar to the static mesh build code; try to share helper structs
+				static TMultiMap<int32, int32> OverlappingRenderingVertexIndices;
+				{
+					OverlappingRenderingVertexIndices.Reset();
+
+					/** Helper struct for building acceleration structures. */
+					struct FIndexAndZ
+					{
+						float Z;
+						int32 Index;
+
+						FIndexAndZ() 
+						{
+						}
+
+						/** Initialization constructor. */
+						FIndexAndZ(int32 InIndex, FVector V)
+						{
+							Z = 0.30f * V.X + 0.33f * V.Y + 0.37f * V.Z;
+							Index = InIndex;
+						}
+					};
+
+					// Build a temporary array of rendering vertex indices, sorted by their Z value.  This will accelerate
+					// searching through to find duplicates
+					static TArray<FIndexAndZ> RenderingVertexIndicesSortedByZ;
+					{
+						RenderingVertexIndicesSortedByZ.SetNumUninitialized( NumRenderingVertices, false );
+						for( int32 RenderingVertexIndex = 0; RenderingVertexIndex < NumRenderingVertices; ++RenderingVertexIndex )
+						{
+							const FVector VertexPosition = StaticMeshLOD.PositionVertexBuffer.VertexPosition( RenderingVertexIndex );
+							RenderingVertexIndicesSortedByZ[ RenderingVertexIndex ] = FIndexAndZ( RenderingVertexIndex, VertexPosition );
+						}
+
+						// Sort the vertices by z value
+						struct FCompareIndexAndZ
+						{
+							FORCEINLINE bool operator()( FIndexAndZ const& A, FIndexAndZ const& B ) const { return A.Z < B.Z; }
+						};
+						RenderingVertexIndicesSortedByZ.Sort( FCompareIndexAndZ() );
+					}
+
+					// Search for duplicates, quickly!
+					const float ComparisonThreshold = KINDA_SMALL_NUMBER;	// @todo mesheditor: Tweak "weld" threshold
+					for( int32 RenderingVertexIterA = 0; RenderingVertexIterA < NumRenderingVertices; ++RenderingVertexIterA )
+					{
+						// only need to search forward, since we add pairs both ways
+						for( int32 RenderingVertexIterB = RenderingVertexIterA + 1; RenderingVertexIterB < NumRenderingVertices; ++RenderingVertexIterB )
+						{
+							if( FMath::Abs( RenderingVertexIndicesSortedByZ[ RenderingVertexIterB ].Z - RenderingVertexIndicesSortedByZ[ RenderingVertexIterA ].Z ) > ComparisonThreshold )
+							{
+								break; // can't be any more dups
+							}
+
+							const int32 RenderingVertexIndexA = RenderingVertexIndicesSortedByZ[ RenderingVertexIterA ].Index;
+							const int32 RenderingVertexIndexB = RenderingVertexIndicesSortedByZ[ RenderingVertexIterB ].Index;
+
+							const FVector VertexPositionA = StaticMeshLOD.PositionVertexBuffer.VertexPosition( RenderingVertexIndexA );
+							const FVector VertexPositionB = StaticMeshLOD.PositionVertexBuffer.VertexPosition( RenderingVertexIndexB );
+
+							if( VertexPositionA.Equals( VertexPositionB, ComparisonThreshold ) )
+							{
+								OverlappingRenderingVertexIndices.Add( RenderingVertexIndexA, RenderingVertexIndexB );
+								OverlappingRenderingVertexIndices.Add( RenderingVertexIndexB, RenderingVertexIndexA );
+							}
+						}
+					}
+				}
+
+				// We'll now make sure we have an editable mesh vertex created for every uniquely-positioned rendering vertex.
+				// Note that it's important that we process all vertices, not only the vertices that are referenced by triangles
+				// in the index buffer, because we properly support meshes with vertices that are not yet connected to any
+				// polygons.  These vertices will simply not have editable mesh polygons or edges connected to them, but will
+				// still be interactable in the editor.
+				for( int32 RenderingVertexIndex = 0; RenderingVertexIndex < NumRenderingVertices; ++RenderingVertexIndex )
+				{
+					const FVector VertexPosition = StaticMeshLOD.PositionVertexBuffer.VertexPosition( RenderingVertexIndex );
+
+					// Check to see if we already have this vertex
+					bool bAlreadyHaveVertexForPosition = false;
+					{
+						static TArray<int32> ThisRenderingVertexOverlaps;
+						ThisRenderingVertexOverlaps.Reset();
+						OverlappingRenderingVertexIndices.MultiFind( RenderingVertexIndex, /* Out */ ThisRenderingVertexOverlaps );
+
+						for( int32 OverlappingRenderingVertexIter = 0; OverlappingRenderingVertexIter < ThisRenderingVertexOverlaps.Num(); ++OverlappingRenderingVertexIter )
+						{
+							const int32 OverlappingRenderingVertexIndex = ThisRenderingVertexOverlaps[ OverlappingRenderingVertexIter ];
+
+							// If the overlapping rendering vertex index is smaller than our current index, we can safely assume that
+							// we've already processed this vertex position and created an editable mesh vertex for it.
+							if( OverlappingRenderingVertexIndex < RenderingVertexIndex )
+							{
+								check( RenderingVertices.IsAllocated( OverlappingRenderingVertexIndex ) );
+								const FVertexID ExistingVertexID = RenderingVertices[ OverlappingRenderingVertexIndex ].VertexID;
+								FEditableStaticMeshVertex& ExistingVertex = Vertices[ ExistingVertexID.GetValue() ];
+
+								// We already have a unique editable vertex for this rendering vertex position, so link them!
+								RenderingVertices.InsertUninitialized( RenderingVertexIndex );
+								RenderingVertices[ RenderingVertexIndex ].VertexID = ExistingVertexID;
+
+								const FRenderingVertexID RenderingVertexID( RenderingVertexIndex );
+								checkSlow( !ExistingVertex.RenderingVertexIDs.Contains( RenderingVertexID ) );
+								ExistingVertex.RenderingVertexIDs.Add( RenderingVertexID );
+								bAlreadyHaveVertexForPosition = true;
+
+								break;
+							}
+						}
+					}
+
+					if( !bAlreadyHaveVertexForPosition )
+					{
+						const FVertexID NewVertexID( Vertices.Add( FEditableStaticMeshVertex() ) );
+						RenderingVertices.InsertUninitialized( RenderingVertexIndex );
+						RenderingVertices[ RenderingVertexIndex ].VertexID = NewVertexID;
+						FEditableStaticMeshVertex& NewVertex = Vertices[ NewVertexID.GetValue() ];
+						NewVertex.VertexPosition = VertexPosition;
+						
+						NewVertex.CornerSharpness = 0.0f;
+
+						// @todo mesheditor: If a mesh somehow contained rendering vertices that no triangle was referencing, this would cause
+						// the rendering vertex to be ignored by the editable mesh code.  It would just sit in the vertex buffer (and in the
+						// editable mesh vertex's RenderingVertexIndices list), but would never be touched.  The editable mesh code only
+						// creates rendering vertices for vertices that are attached to polygons, so this should never happen with meshes
+						// that we create and save.  Only if the incoming data had orphan vertices in it.  Should hopefully not be a problem.
+						const FRenderingVertexID RenderingVertexID( RenderingVertexIndex );
+						NewVertex.RenderingVertexIDs.Add( RenderingVertexID );
+
+						// NOTE: The new vertex's connected polygons will be filled in down below, as we're processing mesh triangles
+					}
+				}
+
+
+				const FIndexArrayView RenderingIndices = StaticMeshLOD.IndexBuffer.GetArrayView();
+
+
+				static TMap<uint64, FEdgeID> UniqueEdgeToEdgeID;
+				UniqueEdgeToEdgeID.Reset();
+
+				// Add all sections
+				const uint32 NumSections = StaticMeshLOD.Sections.Num();
+				for( uint32 RenderingSectionIndex = 0; RenderingSectionIndex < NumSections; ++RenderingSectionIndex )
+				{
+					const FStaticMeshSection& RenderingSection = StaticMeshLOD.Sections[ RenderingSectionIndex ];
+
+					// Create a new editable mesh section
+					const FSectionID NewSectionID( Sections.Add( FEditableStaticMeshSection() ) );
+					FEditableStaticMeshSection& NewSection = Sections[ NewSectionID.GetValue() ];
+					NewSection.RenderingSectionIndex = RenderingSectionIndex;
+					NewSection.MaterialIndex = RenderingSection.MaterialIndex;
+					NewSection.bEnableCollision = RenderingSection.bEnableCollision;
+					NewSection.bCastShadow = RenderingSection.bCastShadow;
+
+					const uint32 NumSectionTriangles = RenderingSection.NumTriangles;
+					NewSection.Triangles.Reserve( NumSectionTriangles );
+					NewSection.MaxTriangles = NumSectionTriangles;
+
+					for( uint32 SectionTriangleIndex = 0; SectionTriangleIndex < NumSectionTriangles; ++SectionTriangleIndex )
+					{
+						const uint32 RenderingTriangleFirstVertexIndex = FEditableStaticMeshSection::TriangleIndexToRenderingTriangleFirstIndex( RenderingSection, FTriangleID( SectionTriangleIndex ) );
+
+						uint32 TriangleRenderingVertexIndices[ 3 ];
+						FVertexID TriangleVertexIDs[ 3 ];
+						for( uint32 TriangleVertexIndex = 0; TriangleVertexIndex < 3; ++TriangleVertexIndex )
+						{
+							TriangleRenderingVertexIndices[ TriangleVertexIndex ] = RenderingIndices[ RenderingTriangleFirstVertexIndex + TriangleVertexIndex ];
+							TriangleVertexIDs[ TriangleVertexIndex ] = RenderingVertices[ TriangleRenderingVertexIndices[ TriangleVertexIndex ] ].VertexID;
+						}
+
+						// Make sure we have a valid triangle.  The triangle can be invalid because at least two if it's vertex indices 
+						// point to the exact same vertex.  The triangle is degenerate.  This can happen due to us welding the overlapping 
+						// vertices because they were either extremely close to each other (or exactly overlapping.)  We'll ignore this triangle.
+						const bool bIsValidTriangle =
+							TriangleVertexIDs[ 0 ] != TriangleVertexIDs[ 1 ] &&
+							TriangleVertexIDs[ 1 ] != TriangleVertexIDs[ 2 ] &&
+							TriangleVertexIDs[ 2 ] != TriangleVertexIDs[ 0 ];
+						if( bIsValidTriangle )
+						{
+							// Static meshes only support triangles, so there's no need to triangulate anything yet.  We'll make both
+							// a triangle and a polygon here.
+							const int32 NewTriangleIndex = SectionTriangleIndex;
+							NewSection.Triangles.InsertUninitialized( NewTriangleIndex );
+							FEditableStaticMeshTriangle& NewTriangle = NewSection.Triangles[ NewTriangleIndex ];
+
+							const FPolygonID NewPolygonID( NewSection.Polygons.Add( FEditableStaticMeshPolygon() ) );
+							FEditableStaticMeshPolygon& NewPolygon = NewSection.Polygons[ NewPolygonID.GetValue() ];
+							NewPolygon.TriangulatedPolygonTriangleIndices.Add( FTriangleID( NewTriangleIndex ) );
+
+
+							// Static meshes don't support polygons with holes, so we always start out with only a perimeter contour per polygon
+							FEditableStaticMeshPolygonContour& PerimeterContour = NewPolygon.PerimeterContour;
+							PerimeterContour.Vertices.Reserve( 3 );
+
+							// Connect vertices
+							for( uint32 TriangleVertexIndex = 0; TriangleVertexIndex < 3; ++TriangleVertexIndex )
+							{
+								const uint32 RenderingVertexIndex = TriangleRenderingVertexIndices[ TriangleVertexIndex ];
+								const FVertexID VertexID = TriangleVertexIDs[ TriangleVertexIndex ];
+								const FRenderingVertexID RenderingVertexID = FRenderingVertexID( RenderingVertexIndex );
+
+								// The triangle points to each of its three vertices
+								NewTriangle.RenderingVertexIDs[ TriangleVertexIndex ] = RenderingVertexID;
+
+								// Tell the polygon contour about this vertex
+								const int32 PolygonContourVertexIndex = PerimeterContour.Vertices.Add( FEditableStaticMeshPolygonContourVertex() );
+								FEditableStaticMeshPolygonContourVertex& PolygonContourVertex = PerimeterContour.Vertices[ PolygonContourVertexIndex ];
+								PolygonContourVertex.VertexID = VertexID;
+								PolygonContourVertex.RenderingVertexID = RenderingVertexID;
+								PolygonContourVertex.VertexUVs.Reserve( NumUVs );
+								for( int32 UVIndex = 0; UVIndex < NumUVs; ++UVIndex )
+								{
+									PolygonContourVertex.VertexUVs.Add( StaticMeshLOD.VertexBuffer.GetVertexUV( RenderingVertexIndex, UVIndex ) );
+								}
+								const FVector Normal = StaticMeshLOD.VertexBuffer.VertexTangentZ( RenderingVertexIndex );
+								const FVector Tangent = StaticMeshLOD.VertexBuffer.VertexTangentX( RenderingVertexIndex );
+								const FVector Binormal = StaticMeshLOD.VertexBuffer.VertexTangentY( RenderingVertexIndex );
+								PolygonContourVertex.Normal = Normal;
+								PolygonContourVertex.Tangent = Tangent;
+								PolygonContourVertex.BinormalSign = GetBasisDeterminantSign(Tangent, Binormal, Normal);
+
+								if( bHasColor )
+								{
+									PolygonContourVertex.Color = FLinearColor(StaticMeshLOD.ColorVertexBuffer.VertexColor( RenderingVertexIndex ) );
+								}
+								else
+								{
+									PolygonContourVertex.Color = FLinearColor::White;
+								}
+							}
+
+
+							// Connect edges
+							{
+								struct Local
+								{
+									inline static uint64 Make64BitValueForEdge( const FVertexID EdgeVertexID0, const FVertexID EdgeVertexID1 )
+									{
+										return uint64( ( uint64( EdgeVertexID0.GetValue() ) << 32 ) | uint64( EdgeVertexID1.GetValue() ) );
+									};
+								};
+
+
+								// Add the edges of this triangle
+								for( uint32 TriangleEdgeNumber = 0; TriangleEdgeNumber < 3; ++TriangleEdgeNumber )
+								{
+									uint32 EdgeRenderingVertexIndices[ 2 ];
+									EdgeRenderingVertexIndices[ 0 ] = RenderingIndices[ RenderingTriangleFirstVertexIndex + ( TriangleEdgeNumber + 0 ) % 3 ];
+									EdgeRenderingVertexIndices[ 1 ] = RenderingIndices[ RenderingTriangleFirstVertexIndex + ( TriangleEdgeNumber + 1 ) % 3 ];
+
+									FVertexID EdgeVertexIDs[ 2 ];
+									EdgeVertexIDs[ 0 ] = RenderingVertices[ EdgeRenderingVertexIndices[ 0 ] ].VertexID;
+									EdgeVertexIDs[ 1 ] = RenderingVertices[ EdgeRenderingVertexIndices[ 1 ] ].VertexID;
+
+									// Check to see if this edge already exists
+									bool bAlreadyHaveEdge = false;
+									FEdgeID EdgeID = FEdgeID::Invalid;
+									{
+										FEdgeID* FoundEdgeIDPtr = UniqueEdgeToEdgeID.Find( Local::Make64BitValueForEdge( EdgeVertexIDs[ 0 ], EdgeVertexIDs[ 1 ] ) );
+										if( FoundEdgeIDPtr == nullptr )
+										{
+											// Try the other way around
+											FoundEdgeIDPtr = UniqueEdgeToEdgeID.Find( Local::Make64BitValueForEdge( EdgeVertexIDs[ 1 ], EdgeVertexIDs[ 0 ] ) );
+										}
+										if( FoundEdgeIDPtr != nullptr )
+										{
+											bAlreadyHaveEdge = true;
+											EdgeID = *FoundEdgeIDPtr;
+										}
+									}
+
+									if( !bAlreadyHaveEdge )
+									{
+										// Create the new edge.  We'll connect it to its polygons later on.
+										CreateEdge_Internal( EdgeVertexIDs[ 0 ], EdgeVertexIDs[ 1 ], TArray<FPolygonRef>(), FEdgeID::Invalid, /* Out */ EdgeID );
+
+										UniqueEdgeToEdgeID.Add( Local::Make64BitValueForEdge( EdgeVertexIDs[ 0 ], EdgeVertexIDs[ 1 ] ), EdgeID );
+									}
+
+									// Each edge will point back to the polygon that its connected to.  Remember, an edge can be shared by multiple
+									// polygons, but usually its best if only shared by up to two.
+									FEditableStaticMeshEdge& Edge = Edges[ EdgeID.GetValue() ];
+
+									Edge.ConnectedPolygons.AddUnique( FPolygonRef( NewSectionID, NewPolygonID ) );
+								}
+							}
+						}
+						else
+						{
+							// Triangle was not valid.  This will result in an empty entry in our Triangles sparse array.  Luckily,
+							// the triangle is already degenerate so we don't need to change anything.  This triangle index will be
+							// re-used if a new triangle needs to be created during editing
+							// @todo mesheditor: This can cause rendering vertices to be orphaned.  Should we delete them?
+						}
+					}
+				}
+
+
+
+				// Figure out whether each edge is hard or soft by looking at their connected polygons
+				for( FEditableStaticMeshEdge& Edge : Edges )
+				{
+					// Default to a hard edge if we have nothing connected
+					Edge.bIsHardEdge = true;
+
+					// Only edges with at least two polygons connected can possibly be soft
+					if( Edge.ConnectedPolygons.Num() >= 2 )
+					{
+						Edge.bIsHardEdge = false;
+
+						FVector FirstEdgeVertexNormals[ 2 ];
+						for( int32 ConnectedPolygonNumber = 0; ConnectedPolygonNumber < Edge.ConnectedPolygons.Num(); ++ConnectedPolygonNumber )
+						{
+							const FPolygonRef& ConnectedPolygonRef = Edge.ConnectedPolygons[ ConnectedPolygonNumber ];
+
+							const FEditableStaticMeshPolygon& Polygon = Sections[ ConnectedPolygonRef.SectionID.GetValue() ].Polygons[ ConnectedPolygonRef.PolygonID.GetValue() ];
+
+							bool bFoundEdge0 = false;
+							bool bFoundEdge1 = false;
+							FVector EdgeVertexNormals[ 2 ];
+							for( int32 VertexNumber = 0; VertexNumber < Polygon.PerimeterContour.Vertices.Num(); ++VertexNumber )
+							{
+								// Find the two vertices for this edge
+								const FVertexID VertexID = Polygon.PerimeterContour.Vertices[ VertexNumber ].VertexID;
+								if( VertexID == Edge.VertexIDs[ 0 ] )
+								{
+									EdgeVertexNormals[0] = GetPolygonPerimeterVertexAttribute( ConnectedPolygonRef, VertexNumber, UEditableMeshAttribute::VertexNormal(), 0 );
+									bFoundEdge0 = true;
+								}
+								else if( VertexID == Edge.VertexIDs[ 1 ] )
+								{
+									EdgeVertexNormals[1] = GetPolygonPerimeterVertexAttribute( ConnectedPolygonRef, VertexNumber, UEditableMeshAttribute::VertexNormal(), 0 );
+									bFoundEdge1 = true;
+								}
+							}
+
+							check( bFoundEdge0 && bFoundEdge1 );
+
+							if( ConnectedPolygonNumber == 0 )
+							{
+								FirstEdgeVertexNormals[0] = EdgeVertexNormals[0];
+								FirstEdgeVertexNormals[1] = EdgeVertexNormals[1];
+							}
+							else
+							{
+								const float Dot0 = FVector::DotProduct( FirstEdgeVertexNormals[0], EdgeVertexNormals[0] );
+								const float Dot1 = FVector::DotProduct( FirstEdgeVertexNormals[1], EdgeVertexNormals[1] );
+
+								// @todo mesheditor: only make hard edges if they are hard in the original model
+								const float MinDotProductForSoftEdge = 0.94f; // adjacent faces with 20 degrees between them get a soft edge
+								if( Dot0 < MinDotProductForSoftEdge || Dot1 < MinDotProductForSoftEdge )
+								{
+									Edge.bIsHardEdge = true;
+									break;
+								}
+							}
+						}
+					}
+				}
+			}
+		}
+	}
+
+	RefreshOpenSubdiv();
+}
+
+
+void UEditableStaticMesh::InitFromBlankStaticMesh( UStaticMesh& InStaticMesh )
+{
+	StaticMesh = &InStaticMesh;
+}
+
+template <typename T, typename ElementIDType>
+static void CompactSparseArrayElements( TSparseArray<T>& Array, TSparseArray<ElementIDType>& IndexRemap )
+{
+	static_assert( TIsDerivedFrom<ElementIDType, FElementID>::IsDerived, "Remap array type must be derived from FElementID" );
+
+	static TSparseArray<T> NewArray;
+	NewArray.Empty( Array.Num() );
+
+	IndexRemap.Empty( Array.GetMaxIndex() );
+
+	// Add valid elements into a new contiguous sparse array.  Note non-const iterator so we can move elements.
+	for( TSparseArray<T>::TIterator It( Array ); It; ++It )
+	{
+		const int32 OldElementIndex = It.GetIndex();
+
+		// @todo mesheditor: implement TSparseArray::Add( ElementType&& ) to save this obscure approach
+		const int32 NewElementIndex = NewArray.Add( T() );
+		NewArray[ NewElementIndex ] = MoveTemp( *It );
+
+		// Provide an O(1) lookup from old index to new index, used when patching up vertex references afterwards
+		IndexRemap.Insert( OldElementIndex, ElementIDType( NewElementIndex ) );
+	}
+
+	Array = MoveTemp( NewArray );
+}
+
+
+template <typename T, typename ElementIDType>
+static void UncompactSparseArrayElements( TSparseArray<T>& Array, const TSparseArray<ElementIDType>& IndexRemap )
+{
+	static_assert( TIsDerivedFrom<ElementIDType, FElementID>::IsDerived, "Remap array type must be derived from FElementID" );
+
+	static TSparseArray<T> NewArray;
+	NewArray.Empty( IndexRemap.GetMaxIndex() );
+
+	// Add valid elements into a new contiguous sparse array.  Note non-const iterator so we can move elements.
+	for( TSparseArray<T>::TIterator It( Array ); It; ++It )
+	{
+		const int32 OldElementIndex = It.GetIndex();
+
+		check( IndexRemap.IsAllocated( OldElementIndex ) );
+		const int32 NewElementIndex = IndexRemap[ OldElementIndex ].GetValue();
+
+		// @todo mesheditor: implement TSparseArray::Insert( ElementType&& ) to save this obscure approach
+		NewArray.Insert( NewElementIndex, T() );
+		NewArray[ NewElementIndex ] = MoveTemp( *It );
+	}
+
+	Array = MoveTemp( NewArray );
+}
+
+
+template <typename ElementIDType>
+static void InvertRemapTable( TSparseArray<ElementIDType>& InvertedRemapTable, const TSparseArray<ElementIDType>& RemapTable )
+{
+	static_assert( TIsDerivedFrom<ElementIDType, FElementID>::IsDerived, "Remap array type must be derived from FElementID" );
+
+	InvertedRemapTable.Empty( RemapTable.Num() );
+
+	for( TSparseArray<ElementIDType>::TConstIterator It( RemapTable ); It; ++It )
+	{
+		InvertedRemapTable.Insert( It->GetValue(), ElementIDType( It.GetIndex() ) );
+	}
+}
+
+
+class FCompactChange : public FChange
+{
+public:
+
+	/** Constructor */
+	FCompactChange()
+	{
+	}
+
+	// Parent class overrides
+	virtual TUniquePtr<FChange> Execute( UObject* Object ) override
+	{
+		UEditableStaticMesh* EditableStaticMesh = CastChecked<UEditableStaticMesh>( Object );
+		verify( !EditableStaticMesh->AnyChangesToUndo() );
+		EditableStaticMesh->Compact();
+		return EditableStaticMesh->MakeUndo();
+	}
+
+	virtual FString ToString() const override
+	{
+		return FString( TEXT( "Compact" ) );
+	}
+};
+
+
+struct FUncompactChangeInput
+{
+	/** A set of remap tables, specifying how the elements should have their indices remapped */
+	UEditableStaticMesh::FElementIDRemappings ElementIDRemappings;
+};
+
+
+class FUncompactChange : public FChange
+{
+public:
+
+	/** Constructor */
+	FUncompactChange( const FUncompactChangeInput& InitInput )
+		: Input( InitInput )
+	{
+	}
+
+	FUncompactChange( FUncompactChangeInput&& InitInput )
+		: Input( MoveTemp( InitInput ) )
+	{
+	}
+
+	// Parent class overrides
+	virtual TUniquePtr<FChange> Execute( UObject* Object ) override
+	{
+		UEditableStaticMesh* EditableStaticMesh = CastChecked<UEditableStaticMesh>( Object );
+		verify( !EditableStaticMesh->AnyChangesToUndo() );
+		EditableStaticMesh->Uncompact( Input.ElementIDRemappings );
+		return EditableStaticMesh->MakeUndo();
+	}
+
+	virtual FString ToString() const override
+	{
+		return FString( TEXT( "Uncompact" ) );
+	}
+
+private:
+
+	/** The data we need to make this change */
+	FUncompactChangeInput Input;
+};
+
+
+void UEditableStaticMesh::FixUpElementIDs( const FElementIDRemappings& Remappings )
+{
+	for( FEditableStaticMeshVertex& Vertex : Vertices )
+	{
+		// Fix up rendering vertex index references in vertices array
+		for( FRenderingVertexID& RenderingVertexID : Vertex.RenderingVertexIDs )
+		{
+			RenderingVertexID = Remappings.GetRemappedRenderingVertexID( RenderingVertexID );
+		}
+
+		// Fix up edge index references in the vertex array
+		for( FEdgeID& EdgeID : Vertex.ConnectedEdgeIDs )
+		{
+			EdgeID = Remappings.GetRemappedEdgeID( EdgeID );
+		}
+	}
+
+	// Fix up vertex index references in rendering vertex array
+	for( FEditableStaticMeshRenderingVertex& RenderingVertex : RenderingVertices )
+	{
+		RenderingVertex.VertexID = Remappings.GetRemappedVertexID( RenderingVertex.VertexID );
+	}
+
+	for( FEditableStaticMeshEdge& Edge : Edges )
+	{
+		// Fix up vertex index references in Edges array
+		for( int32 Index = 0; Index < 2; Index++ )
+		{
+			Edge.VertexIDs[ Index ] = Remappings.GetRemappedVertexID( Edge.VertexIDs[ Index ] );
+		}
+
+		// Fix up references to section indices
+		for( FPolygonRef& ConnectedPolygon : Edge.ConnectedPolygons )
+		{
+			ConnectedPolygon = Remappings.GetRemappedPolygonRef( ConnectedPolygon );
+		}
+	}
+
+	for( TSparseArray< FEditableStaticMeshSection >::TIterator It( Sections ); It; ++It )
+	{
+		FEditableStaticMeshSection& Section = *It;
+		FSectionID SectionID( It.GetIndex() );
+
+		for( FEditableStaticMeshPolygon& Polygon : Section.Polygons )
+		{
+			// Fix up references to vertex indices in section polygons' contours
+			for( FEditableStaticMeshPolygonContourVertex& ContourVertex : Polygon.PerimeterContour.Vertices )
+			{
+				ContourVertex.VertexID = Remappings.GetRemappedVertexID( ContourVertex.VertexID );
+				ContourVertex.RenderingVertexID = Remappings.GetRemappedRenderingVertexID( ContourVertex.RenderingVertexID );
+			}
+
+			for( FEditableStaticMeshPolygonContour& HoleContour : Polygon.HoleContours )
+			{
+				for( FEditableStaticMeshPolygonContourVertex& ContourVertex : HoleContour.Vertices )
+				{
+					ContourVertex.VertexID = Remappings.GetRemappedVertexID( ContourVertex.VertexID );
+					ContourVertex.RenderingVertexID = Remappings.GetRemappedRenderingVertexID( ContourVertex.RenderingVertexID );
+				}
+			}
+
+			// Fix up references to triangle indices
+			for( FTriangleID& TriangleID : Polygon.TriangulatedPolygonTriangleIndices )
+			{
+				TriangleID = Remappings.GetRemappedTriangleID( SectionID, TriangleID );
+			}
+		}
+
+		for( FEditableStaticMeshTriangle& Triangle : Section.Triangles )
+		{
+			for( int32 Index = 0; Index < 3; ++Index )
+			{
+				Triangle.RenderingVertexIDs[ Index ] = Remappings.GetRemappedRenderingVertexID( Triangle.RenderingVertexIDs[ Index ] );
+			}
+		}
+	}
+
+	// @todo mesheditor: broadcast event with remappings so that any cached element IDs can be fixed up.
+	// Will need to move all the editable mesh structures into UEditableMesh and make them public.
+}
+
+
+void UEditableStaticMesh::InitializeStaticMeshBuildVertex( FStaticMeshBuildVertex& StaticMeshVertex, const FEditableStaticMeshPolygonContourVertex ContourVertex )
+{
+	StaticMeshVertex.Position = Vertices[ ContourVertex.VertexID.GetValue() ].VertexPosition;
+	StaticMeshVertex.TangentX = ContourVertex.Tangent;
+	StaticMeshVertex.TangentY = FVector::CrossProduct( ContourVertex.Normal, ContourVertex.Tangent ).GetSafeNormal() * ContourVertex.BinormalSign;
+	StaticMeshVertex.TangentZ = ContourVertex.Normal;
+	StaticMeshVertex.Color = ContourVertex.Color.ToFColor( true );
+	for( int32 UVIndex = 0; UVIndex < ContourVertex.VertexUVs.Num(); ++UVIndex )
+	{
+		StaticMeshVertex.UVs[ UVIndex ] = ContourVertex.VertexUVs[ UVIndex ];
+	}
+}
+
+
+void UEditableStaticMesh::RebuildRenderMesh()
+{
+	if( !IsBeingModified() )
+	{
+		const bool bRefreshBounds = true;
+		const bool bInvalidateLighting = true;
+		RebuildRenderMeshStart( bRefreshBounds, bInvalidateLighting );
+	}
+
+	RebuildRenderMeshInternal();
+
+	if( !IsBeingModified() )
+	{
+		const bool bUpdateCollision = true;
+		RebuildRenderMeshFinish( bUpdateCollision );
+	}
+}
+
+
+void UEditableStaticMesh::RebuildRenderMeshInternal()
+{
+	// @todo mesheditor urgent subdiv: Saw some editable mesh corruption artifacts when testing subDs in VR
+
+	check( RecreateRenderStateContext.IsValid() );
+
+	FStaticMeshLODResources& StaticMeshLOD = GetStaticMeshLOD();
+
+	// Build new vertex buffers
+	static TArray< FStaticMeshBuildVertex > StaticMeshBuildVertices;
+	StaticMeshBuildVertices.Reset();
+
+	static TArray< uint32 > IndexBuffer;
+	IndexBuffer.Reset();
+
+	StaticMeshLOD.Sections.Empty( Sections.Num() );
+
+	bool bHasColor = false;
+
+	if( IsPreviewingSubdivisions() )
+	{
+		check( GetSubdivisionCount() > 0 );
+
+		const int32 SectionCount = SubdivisionLimitData.Sections.Num();
+
+		// @todo mesheditor subdiv: Only 2 UVs supported for now, just to avoid having to use a dynamic array per vertex; needs a new data layout, probably (SoA)
+		const int32 SubdivMeshTextureCoordinateCount = FMath::Min( GetTextureCoordinateCount(), 2 );
+
+		// The Sections sparse array mirrors the SubdivisionLimitData sections array;
+		// iterate through it in parallel in order to get the material index and other section properties
+		TSparseArray< FEditableStaticMeshSection >::TConstIterator SectionIt( Sections );
+		check( Sections.Num() == SectionCount );
+
+		for( int32 SectionNumber = 0; SectionNumber < SectionCount; ++SectionNumber )
+		{
+			const FEditableStaticMeshSection& Section = *SectionIt;
+			const FSubdivisionLimitSection& SubdivisionSection = SubdivisionLimitData.Sections[ SectionNumber ];
+
+			const int32 SectionTriangleCount = SubdivisionSection.SubdividedQuads.Num() * 2;
+
+			// @todo mesheditor subdiv perf: Ideally, if no topology changed we can just fill vertex data and not touch index buffers
+			const int32 FirstSectionVertexIndex = StaticMeshBuildVertices.Num();
+			StaticMeshBuildVertices.AddUninitialized( SectionTriangleCount * 3 );
+
+			const int32 FirstIndexInSection = IndexBuffer.Num();
+			IndexBuffer.Reserve( IndexBuffer.Num() + SectionTriangleCount * 3 );
+
+			// Create new rendering section
+			StaticMeshLOD.Sections.Add( FStaticMeshSection() );
+			FStaticMeshSection& StaticMeshSection = StaticMeshLOD.Sections.Last();
+
+			StaticMeshSection.FirstIndex = FirstIndexInSection;
+			StaticMeshSection.NumTriangles = SectionTriangleCount;
+			StaticMeshSection.MinVertexIndex = FirstSectionVertexIndex;
+			StaticMeshSection.MaxVertexIndex = FirstSectionVertexIndex + SectionTriangleCount * 3;
+
+			StaticMeshSection.MaterialIndex = Section.MaterialIndex;
+			StaticMeshSection.bEnableCollision = Section.bEnableCollision;
+			StaticMeshSection.bCastShadow = Section.bCastShadow;
+
+			// Fill vertices
+			int32 NextVertexIndex = FirstSectionVertexIndex;
+			for( int32 QuadNumber = 0; QuadNumber < SubdivisionSection.SubdividedQuads.Num(); ++QuadNumber )
+			{
+				const FSubdividedQuad& SubdividedQuad = SubdivisionSection.SubdividedQuads[ QuadNumber ];
+
+				// @todo mesheditor subdiv debug
+// 				GWarn->Logf( TEXT( "Q%i V%i: U:%0.2f, V:%0.2f" ), QuadNumber, 0, SubdividedQuad.QuadVertex0.TextureCoordinate0.X, SubdividedQuad.QuadVertex0.TextureCoordinate0.Y );
+// 				GWarn->Logf( TEXT( "Q%i V%i: U:%0.2f, V:%0.2f" ), QuadNumber, 1, SubdividedQuad.QuadVertex1.TextureCoordinate0.X, SubdividedQuad.QuadVertex1.TextureCoordinate0.Y );
+// 				GWarn->Logf( TEXT( "Q%i V%i: U:%0.2f, V:%0.2f" ), QuadNumber, 2, SubdividedQuad.QuadVertex2.TextureCoordinate0.X, SubdividedQuad.QuadVertex2.TextureCoordinate0.Y );
+// 				GWarn->Logf( TEXT( "Q%i V%i: U:%0.2f, V:%0.2f" ), QuadNumber, 3, SubdividedQuad.QuadVertex3.TextureCoordinate0.X, SubdividedQuad.QuadVertex3.TextureCoordinate0.Y );
+
+				for( int32 TriangleNumber = 0; TriangleNumber < 2; ++TriangleNumber )
+				{
+					for( int32 TriangleVertexNumber = 0; TriangleVertexNumber < 3; ++TriangleVertexNumber )
+					{
+						int32 QuadVertexNumber;
+						if( TriangleNumber == 0 )
+						{
+							QuadVertexNumber = ( TriangleVertexNumber == 0 ) ? 0 : ( TriangleVertexNumber == 1 ? 2 : 1 );
+						}
+						else
+						{
+							QuadVertexNumber = ( TriangleVertexNumber == 0 ) ? 0 : ( TriangleVertexNumber == 1 ? 3 : 2 );
+						}
+
+						const FSubdividedQuadVertex& QuadVertex = SubdividedQuad.GetQuadVertex( QuadVertexNumber );
+
+						const FVector VertexPosition = SubdivisionLimitData.VertexPositions[ QuadVertex.VertexPositionIndex ];
+
+						FStaticMeshBuildVertex& StaticMeshVertex = StaticMeshBuildVertices[ NextVertexIndex ];
+						StaticMeshVertex.Position = VertexPosition;
+						StaticMeshVertex.TangentX = QuadVertex.VertexTangent;
+						StaticMeshVertex.TangentY = QuadVertex.VertexBinormalSign * FVector::CrossProduct( QuadVertex.VertexNormal, QuadVertex.VertexTangent );
+						StaticMeshVertex.TangentZ = QuadVertex.VertexNormal;
+
+						for( int32 UVIndex = 0; UVIndex < SubdivMeshTextureCoordinateCount; ++UVIndex )
+						{
+							StaticMeshVertex.UVs[ UVIndex ] = *( &QuadVertex.TextureCoordinate0 + UVIndex );
+						}
+
+						StaticMeshVertex.Color = QuadVertex.VertexColor;
+						if( StaticMeshVertex.Color != FColor::White )
+						{
+							bHasColor = true;
+						}
+
+						IndexBuffer.Add( NextVertexIndex++ );
+					}
+				}
+			}
+
+			++SectionIt;
+		}
+	}
+	else
+	{
+		StaticMeshBuildVertices.SetNum( RenderingVertices.GetMaxIndex() );
+
+		static TBitArray<> VerticesInitialized;
+		VerticesInitialized.Init( false, RenderingVertices.GetMaxIndex() );
+
+		for( FEditableStaticMeshSection& Section : Sections )
+		{
+			Section.RenderingSectionIndex = StaticMeshLOD.Sections.Num();
+
+			// Create new rendering section
+			StaticMeshLOD.Sections.Add( FStaticMeshSection() );
+			FStaticMeshSection& StaticMeshSection = StaticMeshLOD.Sections.Last();
+
+			StaticMeshSection.FirstIndex = IndexBuffer.Num();
+			StaticMeshSection.NumTriangles = Section.Triangles.GetMaxIndex();
+			check( Section.Triangles.GetMaxIndex() <= Section.MaxTriangles );
+			StaticMeshSection.MaterialIndex = Section.MaterialIndex;
+			StaticMeshSection.bEnableCollision = Section.bEnableCollision;
+			StaticMeshSection.bCastShadow = Section.bCastShadow;
+
+			for( const FEditableStaticMeshPolygon& Polygon : Section.Polygons )
+			{
+				for( const FEditableStaticMeshPolygonContourVertex& ContourVertex : Polygon.PerimeterContour.Vertices )
+				{
+					const int32 RenderingVertexIndex = ContourVertex.RenderingVertexID.GetValue();
+
+					if( !VerticesInitialized[ RenderingVertexIndex ] )
+					{
+						if( ContourVertex.Color != FColor::White )
+						{
+							bHasColor = true;
+						}
+
+						FStaticMeshBuildVertex& StaticMeshVertex = StaticMeshBuildVertices[ RenderingVertexIndex ];
+						InitializeStaticMeshBuildVertex( StaticMeshVertex, ContourVertex );
+						VerticesInitialized[ RenderingVertexIndex ] = true;
+					}
+				}
+
+				for( const FEditableStaticMeshPolygonContour& HoleContour : Polygon.HoleContours )
+				{
+					for( const FEditableStaticMeshPolygonContourVertex& ContourVertex : HoleContour.Vertices )
+					{
+						const int32 RenderingVertexIndex = ContourVertex.RenderingVertexID.GetValue();
+
+						if( !VerticesInitialized[ RenderingVertexIndex ] )
+						{
+							if( ContourVertex.Color != FColor::White )
+							{
+								bHasColor = true;
+							}
+							FStaticMeshBuildVertex& StaticMeshVertex = StaticMeshBuildVertices[ RenderingVertexIndex ];
+							InitializeStaticMeshBuildVertex( StaticMeshVertex, ContourVertex );
+							VerticesInitialized[ RenderingVertexIndex ] = true;
+						}
+					}
+				}
+			}
+
+			if( Section.Triangles.Num() > 0 )
+			{
+				IndexBuffer.Reserve( IndexBuffer.Num() + Section.Triangles.GetMaxIndex() * 3 );
+				uint32 MinIndex = TNumericLimits< uint32 >::Max();
+				uint32 MaxIndex = TNumericLimits< uint32 >::Min();
+
+				// Find the first valid rendering vertex index, so that we have a value we can use for our degenerates
+				check( Section.Triangles.Num() > 0 );
+				const FRenderingVertexID FirstValidRenderingID = Section.Triangles.CreateConstIterator()->RenderingVertexIDs[ 0 ];
+
+				for( int32 TriangleIndex = 0; TriangleIndex < Section.Triangles.GetMaxIndex(); ++TriangleIndex )
+				{
+					if( Section.Triangles.IsAllocated( TriangleIndex ) )
+					{
+						const FEditableStaticMeshTriangle& Triangle = Section.Triangles[ TriangleIndex ];
+						for( int32 TriVert = 0; TriVert < 3; ++TriVert )
+						{
+							const uint32 RenderingVertexIndex = Triangle.RenderingVertexIDs[ TriVert ].GetValue();
+							IndexBuffer.Add( RenderingVertexIndex );
+							MinIndex = FMath::Min( MinIndex, RenderingVertexIndex );
+							MaxIndex = FMath::Max( MaxIndex, RenderingVertexIndex );
+						}
+					}
+					else
+					{
+						IndexBuffer.Add( FirstValidRenderingID.GetValue() );
+						IndexBuffer.Add( FirstValidRenderingID.GetValue() );
+						IndexBuffer.Add( FirstValidRenderingID.GetValue() );
+					}
+				}
+
+				StaticMeshSection.MinVertexIndex = MinIndex;
+				StaticMeshSection.MaxVertexIndex = MaxIndex;
+
+				// Add any index buffer padding.
+				// This can be necessary if we have just loaded an editable mesh which had a MaxTriangles count in the editable mesh section
+				// greater than the sparse array max size (i.e. an extra gap had been reserved for tris).
+				const int32 IndexBufferPadding = Section.MaxTriangles - Section.Triangles.GetMaxIndex();
+				if( IndexBufferPadding > 0 )
+				{
+					IndexBuffer.AddZeroed( IndexBufferPadding * 3 );
+				}
+			}
+			else
+			{
+				// No triangles in this section
+				StaticMeshSection.MinVertexIndex = 0;
+				StaticMeshSection.MaxVertexIndex = 0;
+			}
+		}
+	}
+
+	// Figure out which index buffer stride we need
+	bool bNeeds32BitIndices = false;
+	for( const FStaticMeshSection& StaticMeshSection : StaticMeshLOD.Sections )
+	{
+		if( StaticMeshSection.MaxVertexIndex > TNumericLimits<uint16>::Max() )
+		{
+			bNeeds32BitIndices = true;
+		}
+	}
+	const EIndexBufferStride::Type IndexBufferStride = bNeeds32BitIndices ? EIndexBufferStride::Force32Bit : EIndexBufferStride::Force16Bit;
+
+	StaticMeshLOD.PositionVertexBuffer.Init( StaticMeshBuildVertices );
+	StaticMeshLOD.VertexBuffer.Init( StaticMeshBuildVertices, GetTextureCoordinateCount() );
+
+	if( bHasColor )
+	{
+		StaticMeshLOD.ColorVertexBuffer.Init( StaticMeshBuildVertices );
+	}
+	else
+	{
+		StaticMeshLOD.ColorVertexBuffer.InitFromSingleColor( FColor::White, StaticMeshBuildVertices.Num() );
+	}
+
+	StaticMeshLOD.IndexBuffer.SetIndices( IndexBuffer, IndexBufferStride );
+
+	// @todo mesheditor: support the other index buffer types
+	StaticMeshLOD.ReversedIndexBuffer.SetIndices( TArray< uint32 >(), IndexBufferStride );
+	StaticMeshLOD.DepthOnlyIndexBuffer.SetIndices( TArray< uint32 >(), IndexBufferStride );
+	StaticMeshLOD.ReversedDepthOnlyIndexBuffer.SetIndices( TArray< uint32 >(), IndexBufferStride );
+	StaticMeshLOD.WireframeIndexBuffer.SetIndices( TArray< uint32 >(), IndexBufferStride );
+	StaticMeshLOD.AdjacencyIndexBuffer.SetIndices( TArray< uint32 >(), IndexBufferStride );
+
+	StaticMeshLOD.bHasAdjacencyInfo = false;
+	StaticMeshLOD.bHasDepthOnlyIndices = false;
+	StaticMeshLOD.bHasReversedIndices = false;
+	StaticMeshLOD.bHasReversedDepthOnlyIndices = false;
+	StaticMeshLOD.DepthOnlyNumTriangles = 0;
+}
+
+
+void UEditableStaticMesh::Compact()
+{
+	static FElementIDRemappings Remappings;
+
+	// Compact vertices sparse array, generating a lookup from old to new indices in NewVertexIndexLookup
+	CompactSparseArrayElements( Vertices, Remappings.NewVertexIndexLookup );
+
+	// Compact rendering vertices sparse array, generating a lookup from old to new indices in NewVertexIndexLookup
+	CompactSparseArrayElements( RenderingVertices, Remappings.NewRenderingVertexIndexLookup );
+
+	// Compact edges sparse array, generating a lookup from old to new indices in NewEdgeIndexLookup
+	CompactSparseArrayElements( Edges, Remappings.NewEdgeIndexLookup );
+
+	// Compact sections sparse array, generating a lookup from old to new indices in NewSectionIndexLookup
+	CompactSparseArrayElements( Sections, Remappings.NewSectionIndexLookup );
+
+	Remappings.PerPolygon.Empty( Sections.GetMaxIndex() );
+	for( TSparseArray< FEditableStaticMeshSection >::TIterator It( Sections ); It; ++It )
+	{
+		FEditableStaticMeshSection& Section = *It;
+		const int32 Index = It.GetIndex();
+
+		Remappings.PerPolygon.Insert( Index, FElementIDRemappings::FPerPolygonLookups() );
+
+		// Compact the polygon sparse array in each section
+		CompactSparseArrayElements( Section.Polygons, Remappings.PerPolygon[ Index ].NewPolygonIndexLookup );
+
+		// Compact the triangle sparse array in each section
+		CompactSparseArrayElements( Section.Triangles, Remappings.PerPolygon[ Index ].NewTriangleIndexLookup );
+
+		Section.MaxTriangles = Section.Triangles.GetMaxIndex();
+	}
+
+	FixUpElementIDs( Remappings );
+	RebuildRenderMesh();
+
+	// Prepare the inverse transaction to reverse the compaction
+
+	FUncompactChangeInput UncompactChangeInput;
+	InvertRemapTable( UncompactChangeInput.ElementIDRemappings.NewVertexIndexLookup, Remappings.NewVertexIndexLookup );
+	InvertRemapTable( UncompactChangeInput.ElementIDRemappings.NewRenderingVertexIndexLookup, Remappings.NewRenderingVertexIndexLookup );
+	InvertRemapTable( UncompactChangeInput.ElementIDRemappings.NewEdgeIndexLookup, Remappings.NewEdgeIndexLookup );
+	InvertRemapTable( UncompactChangeInput.ElementIDRemappings.NewSectionIndexLookup, Remappings.NewSectionIndexLookup );
+
+	for( TSparseArray< FEditableStaticMeshSection >::TIterator It( Sections ); It; ++It )
+	{
+		const int32 Index = It.GetIndex();
+		const int32 RemappedIndex = UncompactChangeInput.ElementIDRemappings.NewSectionIndexLookup[ Index ].GetValue();
+
+		UncompactChangeInput.ElementIDRemappings.PerPolygon.Insert( RemappedIndex, FElementIDRemappings::FPerPolygonLookups() );
+		auto& PerPolygon = UncompactChangeInput.ElementIDRemappings.PerPolygon[ RemappedIndex ];
+		InvertRemapTable( PerPolygon.NewPolygonIndexLookup, Remappings.PerPolygon[ Index ].NewPolygonIndexLookup );
+		InvertRemapTable( PerPolygon.NewTriangleIndexLookup, Remappings.PerPolygon[ Index ].NewTriangleIndexLookup );
+	}
+
+	AddUndo( MakeUnique<FUncompactChange>( MoveTemp( UncompactChangeInput ) ) );
+}
+
+
+void UEditableStaticMesh::Uncompact( const FElementIDRemappings& Remappings )
+{
+	// Uncompact vertices sparse array, remapping elements according to NewVertexIndexLookup
+	UncompactSparseArrayElements( Vertices, Remappings.NewVertexIndexLookup );
+
+	// Uncompact vertices sparse array, remapping elements according to NewVertexIndexLookup
+	UncompactSparseArrayElements( RenderingVertices, Remappings.NewRenderingVertexIndexLookup );
+
+	// Uncompact vertices sparse array, remapping elements according to NewVertexIndexLookup
+	UncompactSparseArrayElements( Edges, Remappings.NewEdgeIndexLookup );
+
+	// Uncompact vertices sparse array, remapping elements according to NewVertexIndexLookup
+	UncompactSparseArrayElements( Sections, Remappings.NewSectionIndexLookup );
+
+	for( TSparseArray< FEditableStaticMeshSection >::TIterator It( Sections ); It; ++It )
+	{
+		FEditableStaticMeshSection& Section = *It;
+		const int32 Index = It.GetIndex();
+		check( Remappings.PerPolygon.IsAllocated( Index ) );
+
+		// Uncompact the polygon sparse array in each section
+		UncompactSparseArrayElements( Section.Polygons, Remappings.PerPolygon[ Index ].NewPolygonIndexLookup );
+
+		// Uncompact the triangle sparse array in each section
+		UncompactSparseArrayElements( Section.Triangles, Remappings.PerPolygon[ Index ].NewTriangleIndexLookup );
+
+		Section.MaxTriangles = Section.Triangles.GetMaxIndex();
+	}
+
+	FixUpElementIDs( Remappings );
+	RebuildRenderMesh();
+
+	AddUndo( MakeUnique<FCompactChange>() );
+}
+
+
+void UEditableStaticMesh::Serialize( FArchive& Ar )
+{
+	Super::Serialize( Ar );
+
+	Ar.UsingCustomVersion( FEditableMeshCustomVersion::GUID );
+
+	if( Ar.CustomVer( FEditableMeshCustomVersion::GUID ) >= FEditableMeshCustomVersion::TextureCoordinateAndSubdivisionCounts )
+	{
+		Ar << TextureCoordinateCount;
+		Ar << SubdivisionCount;
+	}
+	else if( Ar.IsLoading() )
+	{
+		TextureCoordinateCount = ( StaticMesh != nullptr && StaticMesh->RenderData != nullptr ) ? StaticMesh->RenderData->LODResources[ 0 ].GetNumTexCoords() : 2;
+		SubdivisionCount = 0;
+	}
+
+	SerializeSparseArray( Ar, Vertices );
+	SerializeSparseArray( Ar, RenderingVertices );	// @todo mesheditor serialization: Should not need to be serialized if we triangulate after load
+	SerializeSparseArray( Ar, Edges );
+	SerializeSparseArray( Ar, Sections );
+}
+
+
+void UEditableStaticMesh::StartModification( const EMeshModificationType MeshModificationType, const EMeshTopologyChange MeshTopologyChange )
+{
+	if( ensure( !IsBeingModified() ) )
+	{
+		bIsBeingModified = true;
+
+		// Should be nothing in the undo stack if we're just starting to modify the mesh now
+		ensure( !this->AnyChangesToUndo() );
+
+		FStartOrEndModificationChangeInput RevertInput;
+		RevertInput.bStartModification = false;
+		RevertInput.MeshModificationType = MeshModificationType;
+		RevertInput.MeshTopologyChange = MeshTopologyChange;
+		AddUndo( MakeUnique<FStartOrEndModificationChange>( MoveTemp( RevertInput ) ) );
+
+		this->CurrentModificationType = MeshModificationType;
+		this->CurrentToplogyChange = MeshTopologyChange;
+
+		// @todo mesheditor debug: Disable noisy mesh editor spew by default (here and elsewhere)
+		// UE_LOG( LogMeshEditingRuntime, Log, TEXT( "UEditableStaticMesh::StartModification START: %s" ), *SubMeshAddress.ToString() );
+		FAutoScopedDurationTimer FunctionTimer;
+
+		// @todo mesheditor undo: We're not using traditional transactions to undo mesh changes yet, but we still want to dirty the mesh package
+		// Also, should we even need the Initializing type? Should we not wait for the first modification before dirtying the package?
+		if( 0 )
+		{
+			this->SetFlags( RF_Transactional );
+			this->Modify();
+
+			StaticMesh->SetFlags( RF_Transactional );
+			StaticMesh->Modify();
+		}
+		else
+		{
+			StaticMesh->MarkPackageDirty();
+		}
+
+		const bool bRefreshBounds = CurrentModificationType == EMeshModificationType::Final;	 // @todo mesheditor perf: Only do this if we may have changed the bounds
+		const bool bInvalidateLighting = ( CurrentModificationType == EMeshModificationType::FirstInterim || CurrentModificationType == EMeshModificationType::Final );	// @todo mesheditor perf: We can avoid invalidating lighting on 'Final' if we know that a 'FirstInterim' happened since the last 'Final'
+		RebuildRenderMeshStart( bRefreshBounds, bInvalidateLighting );	
+
+		// @todo mesheditor debug
+		// UE_LOG( LogMeshEditingRuntime, Log, TEXT( "UEditableStaticMesh::StartModification COMPLETE in %0.4fs" ), FunctionTimer.GetTime() );
+	}
+}
+
+
+void UEditableStaticMesh::RebuildRenderMeshStart( const bool bRefreshBounds, const bool bInvalidateLighting )
+{
+	// We're changing the mesh itself, so ALL static mesh components in the scene will need
+	// to be unregistered for this (and reregistered afterwards.)
+	RecreateRenderStateContext = MakeShareable( new FStaticMeshComponentRecreateRenderStateContext( StaticMesh, bInvalidateLighting, bRefreshBounds ) );
+
+	// Release the static mesh's resources.
+	StaticMesh->ReleaseResources();
+
+	// Flush the resource release commands to the rendering thread to ensure that the build doesn't occur while a resource is still
+	// allocated, and potentially accessing the UStaticMesh.
+	StaticMesh->ReleaseResourcesFence.Wait();
+}
+
+
+void UEditableStaticMesh::EndModification( const bool bFromUndo )
+{
+	if( ensure( IsBeingModified() ) )
+	{
+		// @todo mesheditor debug
+		// UE_LOG( LogMeshEditingRuntime, Log, TEXT( "UEditableStaticMesh::EndModification START (ModType=%i): %s" ), (int32)MeshModificationType, *SubMeshAddress.ToString() );
+		// FAutoScopedDurationTimer FunctionTimer;
+
+		if( CurrentModificationType == EMeshModificationType::Final || !bFromUndo )
+		{
+			// Update subdivision limit surface
+			if( CurrentToplogyChange == EMeshTopologyChange::TopologyChange )
+			{
+				// Mesh topology (or subdivision level or smoothing) may have changed, so go ahead and refresh our OpenSubdiv representation entirely
+				RefreshOpenSubdiv();
+			}
+			else
+			{
+				// No topology change, so we can ask OpenSubdiv to quickly generate new limit surface geometry
+				GenerateOpenSubdivLimitSurfaceData();
+			}
+		}
+
+		// Every so often, compact the data.
+		// Note we only want to do this when actions are performed, not when they are being undone/redone
+		bool bDidCompact = false;
+		// @todo mesheditor: reinstate this block once we have refactored to 'uber mesh', so we can implement a callback to notify
+		// tools that IDs have changed.
+		// @todo mesheditor: rendering vertices need to have their indices preserved in FChanges. Suggest refactor to "VertexInstances",
+		// one per unique occurrence of a vertex with equal attributes.
+		if( false )
+		{
+			if( CurrentModificationType == EMeshModificationType::Final &&
+			   CurrentToplogyChange == EMeshTopologyChange::TopologyChange &&
+			   !bFromUndo )
+			{
+				if( ++PendingCompactCounter == CompactFrequency )
+				{
+					PendingCompactCounter = 0;
+					Compact();
+					bDidCompact = true;
+				}
+			}
+		}
+
+		// If subdivision preview mode is active, we'll need to refresh the entire static mesh with data from OpenSubdiv
+		// @todo mesheditor subdiv perf: Ideally we can avoid refreshing the entire thing if only positions have changed, as per above
+		if( IsPreviewingSubdivisions() && ( CurrentModificationType == EMeshModificationType::Final || !bFromUndo ) )
+		{
+			if( !bDidCompact )	// If we did a Compact() in this function, the mesh will have already been rebuilt
+			{
+				RebuildRenderMeshInternal();
+			}
+		}
+
+		RebuildRenderMeshFinish( CurrentModificationType == EMeshModificationType::Final );
+
+		// @todo mesheditor: Not currently sure if we need to do this or not.  Also, we are trying to support runtime editing (no PostEditChange stuff)
+		// 		if( GIsEditor && MeshModificationType == EMeshModificationType::Final )
+		// 		{
+		// 			StaticMesh->PostEditChange();
+		// 		}
+
+		// @todo mesheditor debug
+		// UE_LOG( LogMeshEditingRuntime, Log, TEXT( "UEditableStaticMesh::EndModification COMPLETE in %0.4fs" ), FunctionTimer.GetTime() );	  // @todo mesheditor: Shows bogus time values
+
+		FStartOrEndModificationChangeInput RevertInput;
+		RevertInput.bStartModification = true;
+		RevertInput.MeshModificationType = CurrentModificationType;
+		RevertInput.MeshTopologyChange = CurrentToplogyChange;
+		AddUndo( MakeUnique<FStartOrEndModificationChange>( MoveTemp( RevertInput ) ) );
+
+		bIsBeingModified = false;
+	}
+}
+
+
+void UEditableStaticMesh::RebuildRenderMeshFinish( const bool bUpdateCollision )
+{
+	UpdateBoundsAndCollision( bUpdateCollision );
+
+	StaticMesh->InitResources();
+
+	// NOTE: This can call InvalidateLightingCache() on all components using this mesh, causing Modify() to be 
+	// called on those components!  Just something to be aware of when EndModification() is called within
+	// an undo transaction.
+	RecreateRenderStateContext.Reset();
+}
+
+
+bool UEditableStaticMesh::IsCommitted() const
+{
+	return StaticMesh->EditableMesh == this;
+}
+
+
+bool UEditableStaticMesh::IsCommittedAsInstance() const
+{
+	return StaticMesh != OriginalStaticMesh;
+}
+
+
+void UEditableStaticMesh::Commit()
+{
+	if( !IsCommitted() )
+	{
+		// Move the editable mesh to an inner of the static mesh, and set the static mesh's EditableMesh property.
+		Rename( nullptr, StaticMesh, REN_DontCreateRedirectors );
+		StaticMesh->EditableMesh = this;
+	}
+}
+
+
+UEditableMesh* UEditableStaticMesh::CommitInstance( UPrimitiveComponent* ComponentToInstanceTo )
+{
+	UStaticMeshComponent* StaticMeshComponent = Cast<UStaticMeshComponent>( ComponentToInstanceTo );
+
+	if( StaticMeshComponent )
+	{
+		// Duplicate the static mesh, putting it as an *inner* of the static mesh component.
+		// This is no longer a persistent asset, so clear the appropriate flags.
+		UStaticMesh* NewStaticMesh = DuplicateObject( OriginalStaticMesh, StaticMeshComponent );
+		NewStaticMesh->ClearFlags( RF_Public | RF_Standalone );
+
+		// Point the static mesh component to the new static mesh instance we just made for it
+		StaticMeshComponent->SetStaticMesh( NewStaticMesh );
+
+		// Duplicate this editable mesh to a new instance inside the new static mesh instance, and set the static mesh's EditableMesh property.
+		UEditableStaticMesh* NewEditableMesh = DuplicateObject( this, NewStaticMesh );
+		NewStaticMesh->EditableMesh = NewEditableMesh;
+		NewEditableMesh->StaticMesh = NewStaticMesh;
+
+		// Update the submesh address which will have changed now it's been instanced
+		NewEditableMesh->SetSubMeshAddress( UEditableMeshFactory::MakeSubmeshAddress( StaticMeshComponent, SubMeshAddress.LODIndex ) );
+		NewEditableMesh->RebuildRenderMesh();
+
+		return NewEditableMesh;
+	}
+
+	return nullptr;
+}
+
+
+void UEditableStaticMesh::Revert()
+{
+	// @todo
+}
+
+
+UEditableMesh* UEditableStaticMesh::RevertInstance()
+{
+	// @todo
+	return nullptr;
+}
+
+
+void UEditableStaticMesh::PropagateInstanceChanges()
+{
+	if( IsCommittedAsInstance() )
+	{
+		// @todo mesheditor: we can only generate submesh addresses from a component. Since we don't have a component, we create a dummy one.
+		// Not really fond of this.
+		// Explore other possibilities, e.g. constructing a submesh address by hand (although the contents of MeshObjectPtr are supposed to be opaque)
+		UStaticMeshComponent* DummyComponent = NewObject<UStaticMeshComponent>();
+		DummyComponent->SetStaticMesh( OriginalStaticMesh );
+
+		UEditableStaticMesh* NewEditableMesh = DuplicateObject( this, OriginalStaticMesh );
+		OriginalStaticMesh->EditableMesh = NewEditableMesh;
+		NewEditableMesh->StaticMesh = OriginalStaticMesh;
+		NewEditableMesh->SetSubMeshAddress( UEditableMeshFactory::MakeSubmeshAddress( DummyComponent, SubMeshAddress.LODIndex ) );
+		NewEditableMesh->RebuildRenderMesh();
+	}
+}
+
+
+void UEditableStaticMesh::UpdateBoundsAndCollision( const bool bUpdateCollision )
+{
+	// @todo mesheditor: we will need to create a new DDC key once we are able to edit placed instances individually.
+	// Will need to find a way of deriving the key based on the mesh key and an instance number which remains constant,
+	// otherwise we risk filling the DDC with junk (i.e. using vertex positions etc is not scalable).
+
+	// Compute a new bounding box
+	// @todo mesheditor perf: Only do this if the bounds may have changed (need hinting)
+	{
+		FStaticMeshRenderData& StaticMeshRenderData = *StaticMesh->RenderData;
+
+		FBoxSphereBounds BoundingBoxAndSphere;
+
+		// @todo mesheditor LODs: Really we should store the bounds of LOD0 inside the static mesh.  Our editable mesh might be for a different LOD.
+
+		// If we're in subdivision preview mode, use the bounds of the base cage mesh, so that simple collision
+		// queries will always include the base cage, even though the actual mesh geometry might be quite a bit smaller.
+		// This also relies on us specifically querying against the simple collision, which we do in a second pass after
+		// looking for meshes using a complex collision trace.
+		// @todo mesheditor: Ideally we are not storing an inflated bounds here just for base cage editor interaction
+		if( IsPreviewingSubdivisions() )
+		{
+			BoundingBoxAndSphere = ComputeBoundingBoxAndSphere();
+		}
+		else
+		{
+			FBox BoundingBox;
+			BoundingBox.Init();
+
+			// Could improve performance here if necessary:
+			// 1) cache polygon IDs per vertex (in order to quickly reject orphans) and just iterate vertex array; or
+			// 2) cache bounding box per polygon
+			// There are other cases where having polygon adjacency information (1) might be useful, so it's maybe worth considering.
+
+			for( const FEditableStaticMeshSection& Section : Sections )
+			{
+				for( const FEditableStaticMeshPolygon& Polygon : Section.Polygons )
+				{
+					for( const FEditableStaticMeshPolygonContourVertex& Vertex : Polygon.PerimeterContour.Vertices )
+					{
+						BoundingBox += Vertices[ Vertex.VertexID.GetValue() ].VertexPosition;
+					}
+				}
+			}
+
+			BoundingBox.GetCenterAndExtents( /* Out */ BoundingBoxAndSphere.Origin, /* Out */ BoundingBoxAndSphere.BoxExtent );
+
+			// Calculate the bounding sphere, using the center of the bounding box as the origin.
+			BoundingBoxAndSphere.SphereRadius = 0.0f;
+
+			for( const FEditableStaticMeshSection& Section : Sections )
+			{
+				for( const FEditableStaticMeshPolygon& Polygon : Section.Polygons )
+				{
+					for( const FEditableStaticMeshPolygonContourVertex& Vertex : Polygon.PerimeterContour.Vertices )
+					{
+						const FVector VertexPosition = Vertices[ Vertex.VertexID.GetValue() ].VertexPosition;
+
+						BoundingBoxAndSphere.SphereRadius = FMath::Max( ( VertexPosition - BoundingBoxAndSphere.Origin ).Size(), BoundingBoxAndSphere.SphereRadius );
+					}
+				}
+			}
+		}
+	
+		StaticMeshRenderData.Bounds = BoundingBoxAndSphere;
+		StaticMesh->CalculateExtendedBounds();
+	}
+
+
+	// Refresh collision (only if the interaction has finished though -- this is really expensive!)
+	if( bUpdateCollision )
+	{
+		
+		// @todo mesheditor collision: We're wiping the existing simplified collision and generating a simple bounding
+		// box collision, since that's the best we can do without impacting performance.  We always using visibility (complex)
+		// collision for traces while mesh editing (for hover/selection), so simplified collision isn't really important.
+		const bool bRecreateSimplifiedCollision = true;
+
+		if( StaticMesh->BodySetup == nullptr )
+		{
+			StaticMesh->CreateBodySetup();
+		}
+
+		UBodySetup* BodySetup = StaticMesh->BodySetup;
+
+		// NOTE: We don't bother calling Modify() on the BodySetup as EndModification() will rebuild this guy after every undo
+		// BodySetup->Modify();
+
+		if( bRecreateSimplifiedCollision )
+		{
+			if( BodySetup->AggGeom.GetElementCount() > 0 )
+			{
+				BodySetup->RemoveSimpleCollision();
+			}
+		}
+
+		BodySetup->InvalidatePhysicsData();
+
+		if( bRecreateSimplifiedCollision )
+		{
+			const FBoxSphereBounds Bounds = StaticMesh->GetBounds();
+
+			FKBoxElem BoxElem;
+			BoxElem.Center = Bounds.Origin;
+			BoxElem.X = Bounds.BoxExtent.X * 2.0f;
+			BoxElem.Y = Bounds.BoxExtent.Y * 2.0f;
+			BoxElem.Z = Bounds.BoxExtent.Z * 2.0f;
+			BodySetup->AggGeom.BoxElems.Add( BoxElem );
+		}
+
+		// Update all static mesh components that are using this mesh
+		// @todo mesheditor perf: This is a pretty heavy operation, and overlaps with what we're already doing in RecreateRenderStateContext
+		// a little bit.  Ideally we do everything in a single pass.  Furthermore, if this could be updated lazily it would be faster.
+		{
+			for( FObjectIterator Iter( UStaticMeshComponent::StaticClass() ); Iter; ++Iter )
+			{
+				UStaticMeshComponent* StaticMeshComponent = Cast<UStaticMeshComponent>( *Iter );
+				if( StaticMeshComponent->GetStaticMesh() == StaticMesh )
+				{
+					// it needs to recreate IF it already has been created
+					if( StaticMeshComponent->IsPhysicsStateCreated() )
+					{
+						StaticMeshComponent->RecreatePhysicsState();
+					}
+				}
+			}
+		}
+	}
+}
+
+
+int32 UEditableStaticMesh::GetVertexCount() const
+{
+	return Vertices.Num();
+}
+
+
+int32 UEditableStaticMesh::GetVertexArraySize() const
+{
+	return Vertices.GetMaxIndex();
+}
+
+
+bool UEditableStaticMesh::IsValidVertex( const FVertexID VertexID ) const
+{
+	return VertexID.GetValue() >= 0 && VertexID.GetValue() < Vertices.GetMaxIndex() && Vertices.IsAllocated( VertexID.GetValue() );
+}
+
+
+FVector4 UEditableStaticMesh::GetVertexAttribute( const FVertexID VertexID, const FName AttributeName, const int32 AttributeIndex ) const
+{
+	const FEditableStaticMeshVertex& Vertex = Vertices[ VertexID.GetValue() ];
+
+	if( AttributeName == UEditableMeshAttribute::VertexPosition() )
+	{
+		check( AttributeIndex == 0 );	// Only one position is supported
+		return FVector4( Vertex.VertexPosition, 0.0f );
+	}
+	else if( AttributeName == UEditableMeshAttribute::VertexCornerSharpness() )
+	{
+		check( AttributeIndex == 0 );	// Only one softness value is supported
+		return FVector4( Vertex.CornerSharpness, 0.0f, 0.0f, 0.0f );
+	}
+
+	checkf( 0, TEXT( "UEditableStaticMesh::GetVertexAttribute() called with unrecognized vertex attribute name: %s (index: %i)" ), *AttributeName.ToString(), AttributeIndex );
+	return FVector4( 0.0f );
+}
+
+
+void UEditableStaticMesh::SetVertexAttribute_Internal( const FVertexID VertexID, const FName AttributeName, const int32 AttributeIndex, const FVector4& NewAttributeValue )
+{
+	FEditableStaticMeshVertex& Vertex = Vertices[ VertexID.GetValue() ];
+	FStaticMeshLODResources& StaticMeshLOD = GetStaticMeshLOD();
+
+	if( AttributeName == UEditableMeshAttribute::VertexPosition() )
+	{
+		check( AttributeIndex == 0 );	// Only one position is supported
+
+		const FVector NewVertexPosition( NewAttributeValue );
+
+		Vertex.VertexPosition = NewVertexPosition;
+
+		if( !IsPreviewingSubdivisions() )
+		{
+			// Set the position of all of the rendering vertices for this editable vertex
+			for( const FRenderingVertexID RenderingVertexID : Vertex.RenderingVertexIDs )
+			{
+				check( RenderingVertices.IsAllocated( RenderingVertexID.GetValue() ) );
+				FVector& RenderingVertexPosition = StaticMeshLOD.PositionVertexBuffer.VertexPosition( RenderingVertexID.GetValue() );
+				RenderingVertexPosition = NewVertexPosition;
+			}
+		}
+	}
+	else if( AttributeName == UEditableMeshAttribute::VertexCornerSharpness() )
+	{
+		// @todo mesheditor urgent: We need to use nice logged errors instead of checks() for many of the calls in this module, so that script users don't crash their editor while trying things out.
+		check( AttributeIndex == 0 );	// Only one softness value is supported
+
+		Vertex.CornerSharpness = NewAttributeValue.X;
+	}
+	else
+	{
+		checkf( 0, TEXT( "UEditableStaticMesh::SetVertexAttribute() called with unrecognized vertex attribute name: %s (index: %i)" ), *AttributeName.ToString(), AttributeIndex );
+	}
+}
+
+
+int32 UEditableStaticMesh::GetVertexConnectedEdgeCount( const FVertexID VertexID ) const
+{
+	checkSlow( Vertices.IsAllocated( VertexID.GetValue() ) );
+	const FEditableStaticMeshVertex& Vertex = Vertices[ VertexID.GetValue() ];
+	return Vertex.ConnectedEdgeIDs.Num();
+}
+
+
+FEdgeID UEditableStaticMesh::GetVertexConnectedEdge( const FVertexID VertexID, const int32 ConnectedEdgeNumber ) const
+{
+	checkSlow( Vertices.IsAllocated( VertexID.GetValue() ) );
+	const FEditableStaticMeshVertex& Vertex = Vertices[ VertexID.GetValue() ];
+	return Vertex.ConnectedEdgeIDs[ ConnectedEdgeNumber ];
+}
+
+
+int32 UEditableStaticMesh::GetRenderingVertexCount() const
+{
+	return RenderingVertices.Num();
+}
+
+
+int32 UEditableStaticMesh::GetRenderingVertexArraySize() const
+{
+	return RenderingVertices.GetMaxIndex();
+}
+
+
+FVector4 UEditableStaticMesh::GetEdgeAttribute( const FEdgeID EdgeID, const FName AttributeName, const int32 AttributeIndex ) const
+{
+	checkSlow( Edges.IsAllocated( EdgeID.GetValue() ) );
+	const FEditableStaticMeshEdge& Edge = Edges[ EdgeID.GetValue() ];
+
+	if( AttributeName == UEditableMeshAttribute::EdgeIsHard() )
+	{
+		check( AttributeIndex == 0 );	// Only one edge is hard flag is supported
+
+		return FVector4( Edge.bIsHardEdge ? 1.0f : 0.0f, 0.0f, 0.0f, 0.0f );
+	}
+	else if( AttributeName == UEditableMeshAttribute::EdgeCreaseSharpness() )
+	{
+		check( AttributeIndex == 0 );	// Only one edge crease sharpness is supported
+
+		return FVector4( Edge.CreaseSharpness, 0.0f, 0.0f, 0.0f );
+	}
+
+	checkf( 0, TEXT( "UEditableStaticMesh::GetEdgeAttribute() called with unrecognized edge attribute name: %s (index: %i)" ), *AttributeName.ToString(), AttributeIndex );
+	return FVector4( 0.0f );
+}
+
+
+void UEditableStaticMesh::SetEdgeAttribute_Internal( const FEdgeID EdgeID, const FName AttributeName, const int32 AttributeIndex, const FVector4& NewAttributeValue )
+{
+	checkSlow( Edges.IsAllocated( EdgeID.GetValue() ) );
+	FEditableStaticMeshEdge& Edge = Edges[ EdgeID.GetValue() ];
+
+	if( AttributeName == UEditableMeshAttribute::EdgeIsHard() )
+	{
+		check( AttributeIndex == 0 );	// Only one edge is hard flag is supported
+
+		Edge.bIsHardEdge = !FMath::IsNearlyZero( NewAttributeValue.X );
+	}
+	else if( AttributeName == UEditableMeshAttribute::EdgeCreaseSharpness() )
+	{
+		check( AttributeIndex == 0 );	// Only one edge crease sharpness is supported
+
+		Edge.CreaseSharpness = NewAttributeValue.X;
+	}
+	else
+	{
+		checkf( 0, TEXT( "UEditableStaticMesh::SetEdgeAttribute() called with unrecognized edge attribute name: %s (index: %i)" ), *AttributeName.ToString(), AttributeIndex );
+	}
+}
+
+
+int32 UEditableStaticMesh::GetEdgeCount() const
+{
+	return Edges.Num();
+}
+
+
+int32 UEditableStaticMesh::GetEdgeArraySize() const
+{
+	return Edges.GetMaxIndex();
+}
+
+
+bool UEditableStaticMesh::IsValidEdge( const FEdgeID EdgeID ) const
+{
+	return EdgeID.GetValue() >= 0 && EdgeID.GetValue() < Edges.GetMaxIndex() && Edges.IsAllocated( EdgeID.GetValue() );
+}
+
+
+FVertexID UEditableStaticMesh::GetEdgeVertex( const FEdgeID EdgeID, const int32 EdgeVertexNumber ) const
+{
+	checkSlow( EdgeVertexNumber >= 0 && EdgeVertexNumber < 2 );
+	checkSlow( Edges.IsAllocated( EdgeID.GetValue() ) );
+	return Edges[ EdgeID.GetValue() ].VertexIDs[ EdgeVertexNumber ];
+}
+
+
+int32 UEditableStaticMesh::GetEdgeConnectedPolygonCount( const FEdgeID EdgeID ) const
+{
+	checkSlow( Edges.IsAllocated( EdgeID.GetValue() ) );
+	const FEditableStaticMeshEdge& Edge = Edges[ EdgeID.GetValue() ];
+	return Edge.ConnectedPolygons.Num();
+}
+
+
+FPolygonRef UEditableStaticMesh::GetEdgeConnectedPolygon( const FEdgeID EdgeID, const int32 ConnectedPolygonNumber ) const
+{
+	checkSlow( Edges.IsAllocated( EdgeID.GetValue() ) );
+	const FEditableStaticMeshEdge& Edge = Edges[ EdgeID.GetValue() ];
+	const FPolygonRef& PolygonRef = Edge.ConnectedPolygons[ ConnectedPolygonNumber ];
+	return PolygonRef;
+}
+
+
+int32 UEditableStaticMesh::GetSectionCount() const
+{
+	return Sections.Num();
+}
+
+
+int32 UEditableStaticMesh::GetSectionArraySize() const
+{
+	return Sections.GetMaxIndex();
+}
+
+
+bool UEditableStaticMesh::IsValidSection( const FSectionID SectionID ) const
+{
+	return
+		SectionID.GetValue() >= 0 &&
+		SectionID.GetValue() < Sections.Num() &&
+		Sections.IsAllocated( SectionID.GetValue() );
+}
+
+
+int32 UEditableStaticMesh::GetPolygonCount( const FSectionID SectionID ) const
+{
+	checkSlow( Sections.IsAllocated( SectionID.GetValue() ) );
+	const FEditableStaticMeshSection& Section = Sections[ SectionID.GetValue() ];
+	return Section.Polygons.Num();
+}
+
+
+int32 UEditableStaticMesh::GetPolygonArraySize( const FSectionID SectionID ) const
+{
+	checkSlow( Sections.IsAllocated( SectionID.GetValue() ) );
+	const FEditableStaticMeshSection& Section = Sections[ SectionID.GetValue() ];
+	return Section.Polygons.GetMaxIndex();
+}
+
+
+bool UEditableStaticMesh::IsValidPolygon( const FPolygonRef PolygonRef ) const
+{
+	return
+		IsValidSection( PolygonRef.SectionID ) &&
+		PolygonRef.PolygonID.GetValue() >= 0 &&
+		PolygonRef.PolygonID.GetValue() < Sections[ PolygonRef.SectionID.GetValue() ].Polygons.GetMaxIndex() &&
+		Sections[ PolygonRef.SectionID.GetValue() ].Polygons.IsAllocated( PolygonRef.PolygonID.GetValue() );
+}
+
+
+int32 UEditableStaticMesh::GetTriangleCount( const FSectionID SectionID ) const
+{
+	checkSlow( Sections.IsAllocated( SectionID.GetValue() ) );
+	const FEditableStaticMeshSection& Section = Sections[ SectionID.GetValue() ];
+
+	return Section.Triangles.Num();
+}
+
+
+int32 UEditableStaticMesh::GetTriangleArraySize( const FSectionID SectionID ) const
+{
+	checkSlow( Sections.IsAllocated( SectionID.GetValue() ) );
+	const FEditableStaticMeshSection& Section = Sections[ SectionID.GetValue() ];
+
+	return Section.Triangles.GetMaxIndex();
+}
+
+
+int32 UEditableStaticMesh::GetPolygonPerimeterVertexCount( const FPolygonRef PolygonRef ) const
+{
+	checkSlow( Sections.IsAllocated( PolygonRef.SectionID.GetValue() ) );
+	const FEditableStaticMeshSection& Section = Sections[ PolygonRef.SectionID.GetValue() ];
+
+	checkSlow( Section.Polygons.IsAllocated( PolygonRef.PolygonID.GetValue() ) );
+	const FEditableStaticMeshPolygon& Polygon = Section.Polygons[ PolygonRef.PolygonID.GetValue() ];
+
+	return Polygon.PerimeterContour.Vertices.Num();
+}
+
+
+FVertexID UEditableStaticMesh::GetPolygonPerimeterVertex( const FPolygonRef PolygonRef, const int32 PolygonVertexNumber ) const
+{
+	checkSlow( Sections.IsAllocated( PolygonRef.SectionID.GetValue() ) );
+	const FEditableStaticMeshSection& Section = Sections[ PolygonRef.SectionID.GetValue() ];
+
+	checkSlow( Section.Polygons.IsAllocated( PolygonRef.PolygonID.GetValue() ) );
+	const FEditableStaticMeshPolygon& Polygon = Section.Polygons[ PolygonRef.PolygonID.GetValue() ];
+
+	const FEditableStaticMeshPolygonContourVertex& ContourVertex = Polygon.PerimeterContour.Vertices[ PolygonVertexNumber ];
+	return ContourVertex.VertexID;
+}
+
+
+FRenderingVertexID UEditableStaticMesh::GetPolygonPerimeterRenderingVertex( const FPolygonRef PolygonRef, const int32 PolygonVertexNumber ) const
+{
+	checkSlow( Sections.IsAllocated( PolygonRef.SectionID.GetValue() ) );
+	const FEditableStaticMeshSection& Section = Sections[ PolygonRef.SectionID.GetValue() ];
+
+	checkSlow( Section.Polygons.IsAllocated( PolygonRef.PolygonID.GetValue() ) );
+	const FEditableStaticMeshPolygon& Polygon = Section.Polygons[ PolygonRef.PolygonID.GetValue() ];
+
+	const FEditableStaticMeshPolygonContourVertex& ContourVertex = Polygon.PerimeterContour.Vertices[ PolygonVertexNumber ];
+	return ContourVertex.RenderingVertexID;
+}
+
+
+FVector4 UEditableStaticMesh::GetPolygonPerimeterVertexAttribute( const FPolygonRef PolygonRef, const int32 PolygonVertexNumber, const FName AttributeName, const int32 AttributeIndex ) const
+{
+	checkSlow( Sections.IsAllocated( PolygonRef.SectionID.GetValue() ) );
+	const FEditableStaticMeshSection& Section = Sections[ PolygonRef.SectionID.GetValue() ];
+
+	checkSlow( Section.Polygons.IsAllocated( PolygonRef.PolygonID.GetValue() ) );
+	const FEditableStaticMeshPolygon& Polygon = Section.Polygons[ PolygonRef.PolygonID.GetValue() ];
+
+	const FEditableStaticMeshPolygonContourVertex& ContourVertex = Polygon.PerimeterContour.Vertices[ PolygonVertexNumber ];
+	return GetPolygonContourVertexAttribute( ContourVertex, AttributeName, AttributeIndex );
+}
+
+
+FVector4 UEditableStaticMesh::GetPolygonContourVertexAttribute( const FEditableStaticMeshPolygonContourVertex& PolygonContourVertex, const FName AttributeName, const int32 AttributeIndex ) const
+{
+	if( AttributeName == UEditableMeshAttribute::VertexPosition() )
+	{
+		check( AttributeIndex == 0 );	// Only one position is supported
+		return GetVertexAttribute( PolygonContourVertex.VertexID, AttributeName, AttributeIndex );
+	}
+	else if( AttributeName == UEditableMeshAttribute::VertexNormal() )
+	{
+		check( AttributeIndex == 0 );	// Only one normal is supported
+		return FVector4( PolygonContourVertex.Normal, 0.0f );
+	}
+	else if( AttributeName == UEditableMeshAttribute::VertexTangent() )
+	{
+		check( AttributeIndex == 0 );	// Only one tangent is supported
+		return FVector4( PolygonContourVertex.Tangent, 0.0f );
+	}
+	else if( AttributeName == UEditableMeshAttribute::VertexBinormalSign() )
+	{
+		check( AttributeIndex == 0 );	// Only one basis determinant sign is supported
+		return FVector4( PolygonContourVertex.BinormalSign );
+	}
+	else if( AttributeName == UEditableMeshAttribute::VertexTextureCoordinate() )
+	{
+		const int32 TextureCoordinateIndex = AttributeIndex;
+		if( TextureCoordinateIndex < PolygonContourVertex.VertexUVs.Num() )
+		{
+			const FVector2D TextureCoordinate = PolygonContourVertex.VertexUVs[ TextureCoordinateIndex ];
+			return FVector4( TextureCoordinate, FVector2D::ZeroVector );
+		}
+		else
+		{
+			return FVector4( 0.0f );
+		}
+	}
+	else if( AttributeName == UEditableMeshAttribute::VertexColor() )
+	{
+		check( AttributeIndex == 0 );
+		return FVector4( FLinearColor( PolygonContourVertex.Color ) );
+	}
+
+	checkf( 0, TEXT( "UEditableStaticMesh::GetPolygonVertexAttribute() called with unrecognized vertex attribute name: %s (index: %i)" ), *AttributeName.ToString(), AttributeIndex );
+	return FVector4( 0.0f );
+}
+
+
+void UEditableStaticMesh::SetPolygonPerimeterVertexAttribute_Internal( const FPolygonRef PolygonRef, const int32 PolygonVertexNumber, const FName AttributeName, const int32 AttributeIndex, const FVector4& NewAttributeValue )
+{
+	if( AttributeName != UEditableMeshAttribute::VertexPosition() )
+	{
+		MakeDiscreetPolygonPerimeterRenderingVertexIfNeeded( PolygonRef, PolygonVertexNumber );
+	}
+
+	checkSlow( Sections.IsAllocated( PolygonRef.SectionID.GetValue() ) );
+	FEditableStaticMeshSection& Section = Sections[ PolygonRef.SectionID.GetValue() ];
+
+	checkSlow( Section.Polygons.IsAllocated( PolygonRef.PolygonID.GetValue() ) );
+	FEditableStaticMeshPolygon& Polygon = Section.Polygons[ PolygonRef.PolygonID.GetValue() ];
+
+	FEditableStaticMeshPolygonContourVertex& ContourVertex = Polygon.PerimeterContour.Vertices[ PolygonVertexNumber ];
+	SetPolygonContourVertexAttribute( ContourVertex, AttributeName, AttributeIndex, NewAttributeValue );
+}
+
+
+void UEditableStaticMesh::SetPolygonContourVertexAttribute( FEditableStaticMeshPolygonContourVertex& PolygonContourVertex, const FName AttributeName, const int32 AttributeIndex, const FVector4& NewAttributeValue )
+{
+	if( AttributeName == UEditableMeshAttribute::VertexPosition() )
+	{
+		SetVertexAttribute_Internal( PolygonContourVertex.VertexID, AttributeName, AttributeIndex, NewAttributeValue );
+	}
+	else
+	{
+		FStaticMeshLODResources& StaticMeshLOD = GetStaticMeshLOD();
+		const uint32 RenderingVertexIndex = PolygonContourVertex.RenderingVertexID.GetValue();
+
+		if( AttributeName == UEditableMeshAttribute::VertexNormal() ||
+			AttributeName == UEditableMeshAttribute::VertexTangent() ||
+			AttributeName == UEditableMeshAttribute::VertexBinormalSign() )
+		{
+			check( AttributeIndex == 0 );	// Only one normal is supported
+
+			if( AttributeName == UEditableMeshAttribute::VertexNormal() )
+			{
+				const FVector NewNormal( NewAttributeValue );
+				PolygonContourVertex.Normal = NewNormal;
+			}
+			else if( AttributeName == UEditableMeshAttribute::VertexTangent() )
+			{
+				const FVector NewTangent( NewAttributeValue );
+				PolygonContourVertex.Tangent = NewTangent;
+			}
+			if( AttributeName == UEditableMeshAttribute::VertexBinormalSign() )
+			{
+				PolygonContourVertex.BinormalSign = NewAttributeValue.X;
+			}
+
+			if( !IsPreviewingSubdivisions() )
+			{
+				// @todo mesheditor perf: SetVertexTangents() and VertexTangentX/Y() functions actually does a bit of work to compute the basis every time. 
+				// Ideally we can get/set this stuff directly to improve performance.  This became slower after high precision basis values were added.
+				// @todo mesheditor perf: this is even more pertinent now we already have the binormal sign!
+				StaticMeshLOD.VertexBuffer.SetVertexTangents(
+					RenderingVertexIndex,
+					PolygonContourVertex.Tangent,
+					FVector::CrossProduct( PolygonContourVertex.Normal, PolygonContourVertex.Tangent ).GetSafeNormal() * PolygonContourVertex.BinormalSign,
+					PolygonContourVertex.Normal );
+			}
+		}
+		else if( AttributeName == UEditableMeshAttribute::VertexTextureCoordinate() )
+		{
+			const FVector2D NewTextureCoordinate( NewAttributeValue.X, NewAttributeValue.Y );
+			const int32 TextureCoordinateIndex = AttributeIndex;
+
+			if( PolygonContourVertex.VertexUVs.Num() <= TextureCoordinateIndex )
+			{
+				PolygonContourVertex.VertexUVs.SetNum( TextureCoordinateIndex + 1 );
+			}
+			PolygonContourVertex.VertexUVs[ TextureCoordinateIndex ] = NewTextureCoordinate;
+
+			if( !IsPreviewingSubdivisions() )
+			{
+				check( TextureCoordinateIndex < GetTextureCoordinateCount() );
+				StaticMeshLOD.VertexBuffer.SetVertexUV( RenderingVertexIndex, TextureCoordinateIndex, NewTextureCoordinate );
+			}
+		}
+		else if( AttributeName == UEditableMeshAttribute::VertexColor() )
+		{
+			const FLinearColor NewLinearColor( NewAttributeValue.X, NewAttributeValue.Y, NewAttributeValue.Z, NewAttributeValue.W );
+			const FColor NewColor = NewLinearColor.ToFColor( true );
+
+			PolygonContourVertex.Color = NewColor;
+
+			if( !IsPreviewingSubdivisions() )
+			{
+				if( StaticMeshLOD.ColorVertexBuffer.GetNumVertices() != RenderingVertices.GetMaxIndex() )
+				{
+					if( NewLinearColor != FLinearColor::White )
+					{
+						// Until now, we haven't needed a vertex color buffer.
+						// Force one to be generated now that we have a non-white vertex in the mesh.
+						RebuildRenderMeshInternal();
+					}
+				}
+				else
+				{
+					StaticMeshLOD.ColorVertexBuffer.VertexColor( RenderingVertexIndex ) = NewColor;
+				}
+			}
+		}
+		else
+		{
+			checkf( 0, TEXT( "UEditableStaticMesh::SetPolygonPerimeterVertexAttribute() called with unrecognized vertex attribute name: %s (index: %i)" ), *AttributeName.ToString(), AttributeIndex );
+		}
+	}
+}
+
+
+int32 UEditableStaticMesh::GetPolygonHoleCount( const FPolygonRef PolygonRef ) const
+{
+	checkSlow( Sections.IsAllocated( PolygonRef.SectionID.GetValue() ) );
+	const FEditableStaticMeshSection& Section = Sections[ PolygonRef.SectionID.GetValue() ];
+
+	checkSlow( Section.Polygons.IsAllocated( PolygonRef.PolygonID.GetValue() ) );
+	const FEditableStaticMeshPolygon& Polygon = Section.Polygons[ PolygonRef.PolygonID.GetValue() ];
+
+	return Polygon.HoleContours.Num();
+}
+
+
+int32 UEditableStaticMesh::GetPolygonHoleVertexCount( const FPolygonRef PolygonRef, const int32 HoleNumber ) const
+{
+	checkSlow( Sections.IsAllocated( PolygonRef.SectionID.GetValue() ) );
+	const FEditableStaticMeshSection& Section = Sections[ PolygonRef.SectionID.GetValue() ];
+
+	checkSlow( Section.Polygons.IsAllocated( PolygonRef.PolygonID.GetValue() ) );
+	const FEditableStaticMeshPolygon& Polygon = Section.Polygons[ PolygonRef.PolygonID.GetValue() ];
+
+	const FEditableStaticMeshPolygonContour& Contour = Polygon.HoleContours[ HoleNumber ];
+	return Contour.Vertices.Num();
+}
+
+
+FVertexID UEditableStaticMesh::GetPolygonHoleVertex( const FPolygonRef PolygonRef, const int32 HoleNumber, const int32 PolygonVertexNumber ) const
+{
+	checkSlow( Sections.IsAllocated( PolygonRef.SectionID.GetValue() ) );
+	const FEditableStaticMeshSection& Section = Sections[ PolygonRef.SectionID.GetValue() ];
+
+	checkSlow( Section.Polygons.IsAllocated( PolygonRef.PolygonID.GetValue() ) );
+	const FEditableStaticMeshPolygon& Polygon = Section.Polygons[ PolygonRef.PolygonID.GetValue() ];
+
+	const FEditableStaticMeshPolygonContour& Contour = Polygon.HoleContours[ HoleNumber ];
+	return Contour.Vertices[ PolygonVertexNumber ].VertexID;
+}
+
+
+FRenderingVertexID UEditableStaticMesh::GetPolygonHoleRenderingVertex( const FPolygonRef PolygonRef, const int32 HoleNumber, const int32 PolygonVertexNumber ) const
+{
+	checkSlow( Sections.IsAllocated( PolygonRef.SectionID.GetValue() ) );
+	const FEditableStaticMeshSection& Section = Sections[ PolygonRef.SectionID.GetValue() ];
+
+	checkSlow( Section.Polygons.IsAllocated( PolygonRef.PolygonID.GetValue() ) );
+	const FEditableStaticMeshPolygon& Polygon = Section.Polygons[ PolygonRef.PolygonID.GetValue() ];
+
+	const FEditableStaticMeshPolygonContour& Contour = Polygon.HoleContours[ HoleNumber ];
+	return Contour.Vertices[ PolygonVertexNumber ].RenderingVertexID;
+}
+
+
+FVector4 UEditableStaticMesh::GetPolygonHoleVertexAttribute( const FPolygonRef PolygonRef, const int32 HoleNumber, const int32 PolygonVertexNumber, const FName AttributeName, const int32 AttributeIndex ) const
+{
+	checkSlow( Sections.IsAllocated( PolygonRef.SectionID.GetValue() ) );
+	const FEditableStaticMeshSection& Section = Sections[ PolygonRef.SectionID.GetValue() ];
+
+	checkSlow( Section.Polygons.IsAllocated( PolygonRef.PolygonID.GetValue() ) );
+	const FEditableStaticMeshPolygon& Polygon = Section.Polygons[ PolygonRef.PolygonID.GetValue() ];
+
+	const FEditableStaticMeshPolygonContour& Contour = Polygon.HoleContours[ HoleNumber ];
+	const FEditableStaticMeshPolygonContourVertex& ContourVertex = Contour.Vertices[ PolygonVertexNumber ];
+	return GetPolygonContourVertexAttribute( ContourVertex, AttributeName, AttributeIndex );
+}
+
+
+void UEditableStaticMesh::SetPolygonHoleVertexAttribute_Internal( const FPolygonRef PolygonRef, const int32 HoleNumber, const int32 PolygonVertexNumber, const FName AttributeName, const int32 AttributeIndex, const FVector4& NewAttributeValue )
+{
+	if( AttributeName != UEditableMeshAttribute::VertexPosition() )
+	{
+		// @todo mesheditor hole: We need a version of this function for holes.
+//		MakeDiscreetPolygonPerimeterRenderingVertexIfNeeded( PolygonRef, PolygonVertexNumber );
+	}
+
+	checkSlow( Sections.IsAllocated( PolygonRef.SectionID.GetValue() ) );
+	FEditableStaticMeshSection& Section = Sections[ PolygonRef.SectionID.GetValue() ];
+
+	checkSlow( Section.Polygons.IsAllocated( PolygonRef.PolygonID.GetValue() ) );
+	FEditableStaticMeshPolygon& Polygon = Section.Polygons[ PolygonRef.PolygonID.GetValue() ];
+
+	FEditableStaticMeshPolygonContour& Contour = Polygon.HoleContours[ HoleNumber ];
+	FEditableStaticMeshPolygonContourVertex& ContourVertex = Contour.Vertices[ PolygonVertexNumber ];
+	SetPolygonContourVertexAttribute( ContourVertex, AttributeName, AttributeIndex, NewAttributeValue );
+}
+
+
+int32 UEditableStaticMesh::GetPolygonTriangulatedTriangleCount( const FPolygonRef PolygonRef ) const
+{
+	checkSlow( Sections.IsAllocated( PolygonRef.SectionID.GetValue() ) );
+	const FEditableStaticMeshSection& Section = Sections[ PolygonRef.SectionID.GetValue() ];
+
+	checkSlow( Section.Polygons.IsAllocated( PolygonRef.PolygonID.GetValue() ) );
+	const FEditableStaticMeshPolygon& Polygon = Section.Polygons[ PolygonRef.PolygonID.GetValue() ];
+
+	return Polygon.TriangulatedPolygonTriangleIndices.Num();
+}
+
+
+FVector UEditableStaticMesh::GetPolygonTriangulatedTriangleVertexPosition( const FPolygonRef PolygonRef, const int32 PolygonTriangleNumber, const int32 TriangleVertexNumber ) const
+{
+	checkSlow( TriangleVertexNumber >= 0 && TriangleVertexNumber < 3 );
+	checkSlow( Sections.IsAllocated( PolygonRef.SectionID.GetValue() ) );
+	const FEditableStaticMeshSection& Section = Sections[ PolygonRef.SectionID.GetValue() ];
+
+	checkSlow( Section.Polygons.IsAllocated( PolygonRef.PolygonID.GetValue() ) );
+	const FEditableStaticMeshPolygon& Polygon = Section.Polygons[ PolygonRef.PolygonID.GetValue() ];
+
+	const FTriangleID TriangleID = Polygon.TriangulatedPolygonTriangleIndices[ PolygonTriangleNumber ];
+	checkSlow( Section.Triangles.IsAllocated( TriangleID.GetValue() ) );
+	const FEditableStaticMeshTriangle& Triangle = Section.Triangles[ TriangleID.GetValue() ];
+
+	const FRenderingVertexID RenderingVertexID = Triangle.RenderingVertexIDs[ TriangleVertexNumber ];
+	checkSlow( RenderingVertices.IsAllocated( RenderingVertexID.GetValue() ) );
+	const FVertexID VertexID = RenderingVertices[ RenderingVertexID.GetValue() ].VertexID;
+
+	checkSlow( Vertices.IsAllocated( VertexID.GetValue() ) );
+	return Vertices[ VertexID.GetValue() ].VertexPosition;
+}
+
+
+void UEditableStaticMesh::CreateEmptyVertexRange_Internal( const int32 NumVerticesToAdd, const TArray<FVertexID>* OverrideVertexIDsForRedo, TArray<FVertexID>& OutNewVertexIDs )
+{
+	check( NumVerticesToAdd > 0 );
+
+	OutNewVertexIDs.Reset();
+
+	static TArray<FEditableStaticMeshVertex> NewVertices;
+	NewVertices.Reset();
+	NewVertices.AddDefaulted( NumVerticesToAdd );
+
+	for( int32 VertexToAddNumber = 0; VertexToAddNumber < NumVerticesToAdd; ++VertexToAddNumber )
+	{
+		FVertexID NewVertexID;
+		if( OverrideVertexIDsForRedo != nullptr )
+		{
+			NewVertexID = ( *OverrideVertexIDsForRedo )[ VertexToAddNumber ];
+			Vertices.Insert( NewVertexID.GetValue(), FEditableStaticMeshVertex() );
+		}
+		else
+		{
+			NewVertexID = FVertexID( Vertices.Add( FEditableStaticMeshVertex() ) );
+		}
+
+		FEditableStaticMeshVertex& NewVertex = Vertices[ NewVertexID.GetValue() ];
+
+		// Default position
+		NewVertex.VertexPosition = FVector::ZeroVector;
+
+		// Default corner sharpness
+		NewVertex.CornerSharpness = 0.0f;
+
+		// NOTE: The vertex starts out with no rendering vertex indices.  Those will be setup when the vertex is
+		// connected to a polygon
+
+		OutNewVertexIDs.Add( NewVertexID );
+	}
+}
+
+
+void UEditableStaticMesh::CreateEdge_Internal( const FVertexID VertexIDA, const FVertexID VertexIDB, const TArray<FPolygonRef>& ConnectedPolygons, const FEdgeID OverrideEdgeIDForRedo, FEdgeID& OutNewEdgeID )
+{
+	FEdgeID NewEdgeID;
+
+	if( OverrideEdgeIDForRedo != FEdgeID::Invalid )
+	{
+		NewEdgeID = OverrideEdgeIDForRedo;
+		Edges.Insert( NewEdgeID.GetValue(), FEditableStaticMeshEdge() );
+	}
+	else
+	{
+		NewEdgeID = FEdgeID( Edges.Add( FEditableStaticMeshEdge() ) );
+	}
+	FEditableStaticMeshEdge& NewEdge = Edges[ NewEdgeID.GetValue() ];
+
+	NewEdge.VertexIDs[ 0 ] = VertexIDA;
+	NewEdge.VertexIDs[ 1 ] = VertexIDB;
+
+	NewEdge.ConnectedPolygons = ConnectedPolygons;
+
+	NewEdge.bIsHardEdge = false;
+	NewEdge.CreaseSharpness = 0.0f;
+
+	// Connect the edge to its vertices
+	Vertices[ NewEdge.VertexIDs[ 0 ].GetValue() ].ConnectedEdgeIDs.Add( NewEdgeID );
+	Vertices[ NewEdge.VertexIDs[ 1 ].GetValue() ].ConnectedEdgeIDs.Add( NewEdgeID );
+
+	OutNewEdgeID = NewEdgeID;
+}
+
+
+void UEditableStaticMesh::CreatePolygon_Internal( const FSectionID SectionID, const TArray<FVertexID>& VertexIDs, const TArray<TArray<FVertexID>>& VertexIDsForEachHole, const FPolygonID OverridePolygonIDForRedo, FPolygonRef& OutNewPolygonRef, TArray<FEdgeID>& OutNewEdgeIDs )
+{
+	OutNewEdgeIDs.Reset();
+
+	// All polygons must have at least three vertices
+	check( VertexIDs.Num() >= 3 );
+
+	FEditableStaticMeshSection& Section = Sections[ SectionID.GetValue() ];
+
+	FPolygonID NewPolygonID;
+	if( OverridePolygonIDForRedo != FPolygonID::Invalid )
+	{
+		NewPolygonID = OverridePolygonIDForRedo;
+		Section.Polygons.Insert( NewPolygonID.GetValue(), FEditableStaticMeshPolygon() );
+	}
+	else
+	{
+		NewPolygonID = FPolygonID( Section.Polygons.Add( FEditableStaticMeshPolygon() ) );
+	}
+	FEditableStaticMeshPolygon& NewPolygon = Section.Polygons[ NewPolygonID.GetValue() ];
+
+	const FPolygonRef PolygonRef( SectionID, NewPolygonID );
+
+
+	// Set our vertex IDs, then tell all of our edges that we are now connected to them
+	{
+		{
+			// Create new rendering vertices for the polygon.  One for each vertex ID.
+			NewPolygon.PerimeterContour.Vertices.SetNum( VertexIDs.Num(), false );
+			for( int32 PerimeterVertexNumber = 0; PerimeterVertexNumber < VertexIDs.Num(); ++PerimeterVertexNumber )
+			{
+				const FVertexID VertexID = VertexIDs[ PerimeterVertexNumber ];
+				NewPolygon.PerimeterContour.Vertices[ PerimeterVertexNumber ].VertexID = VertexID;
+				const FRenderingVertexID NewRenderingVertexID = AddNewRenderingVertexToPolygonPerimeter( PolygonRef, PerimeterVertexNumber );
+			}
+		}
+
+		// Make sure we have valid edges that connect the incoming vertex IDs.  We'll create any edges that are missing
+		static TArray<FEdgeID> NewEdgeIDsForPolygonPerimeter;
+		NewEdgeIDsForPolygonPerimeter.Reset();
+		CreateMissingPolygonPerimeterEdges( PolygonRef, /* Out */ NewEdgeIDsForPolygonPerimeter );
+
+		OutNewEdgeIDs.Append( NewEdgeIDsForPolygonPerimeter );
+
+		static TArray< FEdgeID > ContourEdgeIDs;
+		GetPolygonPerimeterEdges( PolygonRef, /* Out */ ContourEdgeIDs );
+		for( const FEdgeID EdgeID : ContourEdgeIDs )
+		{
+			FEditableStaticMeshEdge& Edge = Edges[ EdgeID.GetValue() ];
+			FPolygonRef& NewPolygonRef = *new( Edge.ConnectedPolygons ) FPolygonRef( PolygonRef );
+		}
+
+		const int32 NumPolygonHoles = VertexIDsForEachHole.Num();
+		for( int32 HoleNumber = 0; HoleNumber < NumPolygonHoles; ++HoleNumber )
+		{
+			FEditableStaticMeshPolygonContour& HoleContour = *new( NewPolygon.HoleContours ) FEditableStaticMeshPolygonContour();
+
+			{
+				// Create new rendering vertices for the polygon hole.  One for each vertex ID.
+				HoleContour.Vertices.SetNum( VertexIDsForEachHole[ HoleNumber ].Num(), false );
+				for( int32 HoleVertexNumber = 0; HoleVertexNumber < HoleContour.Vertices.Num(); ++HoleVertexNumber )
+				{
+					// @todo mesheditor: support for holes
+					const FVertexID VertexID = HoleContour.Vertices[ HoleVertexNumber ].VertexID;
+					HoleContour.Vertices[ HoleVertexNumber ].VertexID = VertexID;
+
+					// @todo mesheditor holes
+					// const FRenderingVertexID NewRenderingVertexID = AddNewRenderingVertexToPolygonHole( PolygonRef, HoleIndex, HoleVertexNumber );
+				}
+			}
+
+			// Make sure we have valid edges that connect the incoming hole vertex IDs.  We'll create any edges that are missing
+			static TArray<FEdgeID> NewEdgeIDsForPolygonHole;
+			NewEdgeIDsForPolygonHole.Reset();
+			CreateMissingPolygonHoleEdges( PolygonRef, HoleNumber, /* Out */ NewEdgeIDsForPolygonHole );
+
+			OutNewEdgeIDs.Append( NewEdgeIDsForPolygonHole );
+
+			GetPolygonHoleEdges( PolygonRef, HoleNumber, /* Out */ ContourEdgeIDs );
+			for( const FEdgeID EdgeID : ContourEdgeIDs )
+			{
+				FEditableStaticMeshEdge& Edge = Edges[ EdgeID.GetValue() ];
+				FPolygonRef& NewPolygonRef = *new( Edge.ConnectedPolygons ) FPolygonRef( PolygonRef );
+			}
+		}
+	}
+
+	
+	// Generate triangles for the new polygon
+	{
+		static TArray<FPolygonRef> PolygonsToRetriangulate;
+		PolygonsToRetriangulate.Reset();
+		PolygonsToRetriangulate.Add( PolygonRef );
+		const bool bOnlyOnUndo = false;
+
+		// NOTE: We don't keep the revert step for retriangulation, because this DeletePolygons_Internal() is used to clean up
+		// newly-created polygons, so we'll never need to do our own rollback.
+		const bool bWasUndoEnabled = bAllowUndo;
+		SetAllowUndo( false );
+
+		RetriangulatePolygons( PolygonsToRetriangulate, bOnlyOnUndo );
+
+		SetAllowUndo( bWasUndoEnabled );
+	}
+
+	OutNewPolygonRef = PolygonRef;
+}
+
+
+void UEditableStaticMesh::RetriangulatePolygons( const TArray<FPolygonRef>& PolygonRefs, const bool bOnlyOnUndo )
+{
+	FRetrianglulatePolygonsChangeInput RevertInput;
+	RevertInput.PolygonRefs = PolygonRefs;
+	RevertInput.bOnlyOnUndo = !bOnlyOnUndo;
+
+	if( !bOnlyOnUndo )
+	{
+		for( const FPolygonRef PolygonRef : PolygonRefs )
+		{
+			FEditableStaticMeshSection& Section = Sections[ PolygonRef.SectionID.GetValue() ];
+			FEditableStaticMeshPolygon& Polygon = Section.Polygons[ PolygonRef.PolygonID.GetValue() ];
+
+			// @todo mesheditor urgent: Any time a polygon's shape changes (vertices/edges are moved, new vertices added, etc), we should 
+			// wipe the polygon's triangles and call MakeTrianglesForPolygon() again to update the polygon's triangles.  This is critical to
+			// handle cases where the polygon becomes or is no longer degenerate, or convex -> concave transitions and vice versa
+			static TArray<int32> PerimeterVertexNumbersForTriangles;
+			ComputePolygonTriangulation( PolygonRef, /* Out */ PerimeterVertexNumbersForTriangles );
+			check( PerimeterVertexNumbersForTriangles.Num() > 0 );
+
+			// @todo mesheditor holes: This code will need some work to support vertex IDs that come from hole contours
+			static TArray<FRenderingVertexID> TrianglesRenderingVertexIDs;
+			TrianglesRenderingVertexIDs.SetNum( PerimeterVertexNumbersForTriangles.Num(), false );
+			for( int32 TriangleVerticesNumber = 0; TriangleVerticesNumber < PerimeterVertexNumbersForTriangles.Num(); ++TriangleVerticesNumber )
+			{
+				const int32 PerimeterVertexNumber = PerimeterVertexNumbersForTriangles[ TriangleVerticesNumber ];
+				TrianglesRenderingVertexIDs[ TriangleVerticesNumber ] = Polygon.PerimeterContour.Vertices[ PerimeterVertexNumber ].RenderingVertexID;
+			}
+
+			// Check to see whether the index buffer needs to be updated
+			bool bNeedsUpdatedTriangles = false;
+			{
+				if( Polygon.TriangulatedPolygonTriangleIndices.Num() * 3 != TrianglesRenderingVertexIDs.Num() )
+				{
+					// Triangle count has changed, so we definitely need new triangles!
+					bNeedsUpdatedTriangles = true;
+				}
+				else
+				{
+					// @todo mesheditor: Untested code path
+					int32 NextNewTriangleVertexNumber = 0;
+					for( int32 TriangleIter = 0; !bNeedsUpdatedTriangles && TriangleIter < Polygon.TriangulatedPolygonTriangleIndices.Num(); ++TriangleIter )
+					{
+						const FEditableStaticMeshTriangle& OldTriangle = Section.Triangles[ Polygon.TriangulatedPolygonTriangleIndices[ TriangleIter ].GetValue() ];
+						for( int32 TriangleVertexIter = 0; TriangleVertexIter < 3; ++TriangleVertexIter )
+						{
+							if( OldTriangle.RenderingVertexIDs[ TriangleVertexIter ] != TrianglesRenderingVertexIDs[ NextNewTriangleVertexNumber ] )
+							{
+								bNeedsUpdatedTriangles = true;
+								break;
+							}
+							++NextNewTriangleVertexNumber;
+						}
+					}
+				}
+			}
+
+			// Has anything changed?
+			if( bNeedsUpdatedTriangles )
+			{
+				const uint32 RenderingSectionIndex = Section.RenderingSectionIndex;
+				FStaticMeshLODResources& StaticMeshLOD = GetStaticMeshLOD();
+				FStaticMeshSection& RenderingSection = StaticMeshLOD.Sections[ RenderingSectionIndex ];
+
+				// Remove the old triangles
+				DeletePolygonTriangles( PolygonRef );
+
+				// Add new triangles
+				{
+					// This is the number of triangles we are about to add
+					const int32 NumNewTriangles = TrianglesRenderingVertexIDs.Num() / 3;
+
+					// This is the number of entries currently unused in the Triangles sparse array
+					const int32 NumFreeTriangles = Section.Triangles.GetMaxIndex() - Section.Triangles.Num();
+
+					// This is the number of triangles we need to make extra space for (in both the sparse array and the index buffer)
+					const int32 NumTrianglesToReserve = FMath::Max( 0, NumNewTriangles - NumFreeTriangles );
+
+					// This is the number of triangles we will need to have allocated in the index buffer after adding the new triangles
+					const int32 NewTotalTriangles = Section.Triangles.GetMaxIndex() + NumTrianglesToReserve;
+
+					// Reserve extra triangles if necessary.
+					if( NumTrianglesToReserve > 0 )
+					{
+						Section.Triangles.Reserve( NewTotalTriangles );
+					}
+
+					// Keep track of new min/max vertex indices
+					int32 MinVertexIndex = RenderingSection.MinVertexIndex;
+					int32 MaxVertexIndex = RenderingSection.MaxVertexIndex;
+
+					// Create empty triangles for all of the new triangles we need, and keep track of their triangle indices
+					static TArray<int32> NewTriangleIndices;
+					{
+						NewTriangleIndices.SetNumUninitialized( NumNewTriangles, false );
+
+						for( int32 TriangleToAddNumber = 0; TriangleToAddNumber < NumNewTriangles; ++TriangleToAddNumber )
+						{
+							const int32 NewTriangleIndex = Section.Triangles.Add( FEditableStaticMeshTriangle() );
+							NewTriangleIndices[ TriangleToAddNumber ] = NewTriangleIndex;
+
+							FEditableStaticMeshTriangle& NewTriangle = Section.Triangles[ NewTriangleIndex ];
+							for( int32 TriangleVertexNumber = 0; TriangleVertexNumber < 3; ++TriangleVertexNumber )
+							{
+								FRenderingVertexID RenderingVertexID = TrianglesRenderingVertexIDs[ TriangleToAddNumber * 3 + TriangleVertexNumber ];
+								NewTriangle.RenderingVertexIDs[ TriangleVertexNumber ] = RenderingVertexID;
+								MinVertexIndex = FMath::Min( MinVertexIndex, RenderingVertexID.GetValue() );
+								MaxVertexIndex = FMath::Max( MaxVertexIndex, RenderingVertexID.GetValue() );
+							}
+
+							Polygon.TriangulatedPolygonTriangleIndices.Add( FTriangleID( NewTriangleIndex ) );
+						}
+					}
+
+					// Update the index buffer
+					if( !IsPreviewingSubdivisions() )
+					{
+						UpdateIndexBufferFormatIfNeeded( TrianglesRenderingVertexIDs );
+					}
+
+					// If we need more space in the index buffer for this section, allocate it here
+					if( NewTotalTriangles > Section.MaxTriangles )
+					{
+						AllocateExtraIndicesForSection( PolygonRef.SectionID, NewTotalTriangles + IndexBufferInterSectionGap - Section.MaxTriangles );
+					}
+
+					if( !IsPreviewingSubdivisions() )
+					{
+						for( int32 TriangleToAddNumber = 0; TriangleToAddNumber < NumNewTriangles; ++TriangleToAddNumber )
+						{
+							const int32 NewTriangleIndex = NewTriangleIndices[ TriangleToAddNumber ];
+
+							for( int32 TriangleVertexNumber = 0; TriangleVertexNumber < 3; ++TriangleVertexNumber )
+							{
+								StaticMeshLOD.IndexBuffer.SetIndex(
+									FEditableStaticMeshSection::TriangleIndexToRenderingTriangleFirstIndex( RenderingSection, FTriangleID( NewTriangleIndex ) ) + TriangleVertexNumber,
+									TrianglesRenderingVertexIDs[ TriangleToAddNumber * 3 + TriangleVertexNumber ].GetValue() );
+							}
+						}
+
+						if( NumTrianglesToReserve > 0 )
+						{
+							RenderingSection.NumTriangles += NumTrianglesToReserve;
+						}
+						RenderingSection.MinVertexIndex = MinVertexIndex;
+						RenderingSection.MaxVertexIndex = MaxVertexIndex;
+					}
+				}
+			}
+		}
+	}
+
+#if 0 // Debug code for checking that sections are correctly allocated
+	UE_LOG( LogTemp, Log, TEXT( "RetriangulatePolygons" ) );
+	FStaticMeshLODResources& StaticMeshLOD = GetStaticMeshLOD();
+	int32 RenderingSectionIndex = 0;
+	for( const FStaticMeshSection& RenderingSection : StaticMeshLOD.Sections )
+	{
+		const FSectionID SectionID = GetSectionForRenderingSectionIndex( RenderingSectionIndex );
+		check( SectionID != FSectionID::Invalid );
+		const FEditableStaticMeshSection& Section = Sections[ SectionID.GetValue() ];
+		check( Section.MaterialIndex == RenderingSection.MaterialIndex );
+
+		UE_LOG( LogTemp, Log, TEXT( "  %s: Material = %s, FirstIndex = %d, Num tris = %d, Max tris = %d" ),
+			   *SectionID.ToString(),
+			   *StaticMesh->GetMaterial(Section.MaterialIndex)->GetName(),
+			   RenderingSection.FirstIndex,
+			   RenderingSection.NumTriangles,
+			   Section.MaxTriangles );
+
+		++RenderingSectionIndex;
+	}
+#endif
+
+	AddUndo( MakeUnique<FRetrianglulatePolygonsChange>( MoveTemp( RevertInput ) ) );
+}
+
+
+void UEditableStaticMesh::AllocateExtraIndicesForSection( const FSectionID SectionID, int32 NumExtraTriangles )
+{
+	check( Sections.IsAllocated( SectionID.GetValue() ) );
+
+	FEditableStaticMeshSection& Section = Sections[ SectionID.GetValue() ];
+
+	// Get current number of triangles allocated for this section
+	const int32 MaxTriangles = Section.MaxTriangles;
+	Section.MaxTriangles += NumExtraTriangles;
+
+	if( !IsPreviewingSubdivisions() )
+	{
+		FStaticMeshLODResources& StaticMeshLOD = GetStaticMeshLOD();
+		uint32 RenderingSectionIndex = Section.RenderingSectionIndex;
+		uint32 FirstIndex = StaticMeshLOD.Sections[ RenderingSectionIndex ].FirstIndex;
+
+		// Make room in the index buffer for the extra triangles, and update the mesh section's maximum triangle count
+		StaticMeshLOD.IndexBuffer.InsertIndices( FirstIndex + MaxTriangles * 3, nullptr, NumExtraTriangles * 3 );
+
+		// Adjust first index for all subsequent render sections to account for the extra indices just inserted.
+		// It is guaranteed that index buffer indices are in the same order as the rendering sections.
+		const uint32 NumRenderingSections = StaticMeshLOD.Sections.Num();
+		while( ++RenderingSectionIndex < NumRenderingSections )
+		{
+			check( StaticMeshLOD.Sections[ RenderingSectionIndex ].FirstIndex >= FirstIndex );
+			StaticMeshLOD.Sections[ RenderingSectionIndex ].FirstIndex += NumExtraTriangles * 3;
+		}
+	}
+}
+
+
+void UEditableStaticMesh::CreateRenderingVertices( const TArray<FVertexID>& VertexIDs, const TOptional<FRenderingVertexID> OptionalCopyFromRenderingVertexID, TArray<FRenderingVertexID>& OutNewRenderingVertexIDs )
+{
+	const int32 NumVerticesToAdd = VertexIDs.Num();
+
+	int32 NumFreeRenderingVertexIDs = RenderingVertices.GetMaxIndex() - RenderingVertices.Num();
+	check( NumFreeRenderingVertexIDs >= 0 );
+
+	RenderingVertices.Reserve( RenderingVertices.Num() + NumVerticesToAdd );
+
+	OutNewRenderingVertexIDs.Reset();
+	OutNewRenderingVertexIDs.Reserve( NumVerticesToAdd );
+	for( int32 VertexToAddNumber = 0; VertexToAddNumber < NumVerticesToAdd; ++VertexToAddNumber )
+	{
+		const int32 NewRenderingVertexIndex = RenderingVertices.AddUninitialized().Index;
+		RenderingVertices[ NewRenderingVertexIndex ].VertexID = VertexIDs[ VertexToAddNumber ];
+
+		const FRenderingVertexID NewRenderingVertexID( NewRenderingVertexIndex );
+		OutNewRenderingVertexIDs.Add( NewRenderingVertexID );
+
+		// Update the vertex
+		FEditableStaticMeshVertex& ReferencedVertex = Vertices[ VertexIDs[ VertexToAddNumber ].GetValue() ];
+		checkSlow( !ReferencedVertex.RenderingVertexIDs.Contains( NewRenderingVertexID ) );
+		ReferencedVertex.RenderingVertexIDs.Add( NewRenderingVertexID );
+	}
+
+	if( !IsPreviewingSubdivisions() )
+	{
+		FStaticMeshLODResources& StaticMeshLOD = GetStaticMeshLOD();
+		const int32 NumUVs = GetTextureCoordinateCount();
+		const bool bHasColors = StaticMeshLOD.ColorVertexBuffer.GetNumVertices() > 0;
+
+		const int32 OldVertexBufferRenderingVertexCount = StaticMeshLOD.PositionVertexBuffer.GetNumVertices();
+		const int32 NumNewVertexBufferRenderingVertices = FMath::Max( 0, NumVerticesToAdd - NumFreeRenderingVertexIDs );
+
+		static TArray<FStaticMeshBuildVertex> RenderingVerticesToAppend;
+		RenderingVerticesToAppend.SetNumUninitialized( NumNewVertexBufferRenderingVertices, false );
+
+		for( int32 VertexToAddNumber = 0; VertexToAddNumber < NumVerticesToAdd; ++VertexToAddNumber )
+		{
+			const FEditableStaticMeshVertex& ReferencedVertex = Vertices[ VertexIDs[ VertexToAddNumber ].GetValue() ];
+
+			const FRenderingVertexID NewRenderingVertexID = OutNewRenderingVertexIDs[ VertexToAddNumber ];
+			const int32 NewRenderingVertexIndex = NewRenderingVertexID.GetValue();
+
+			if( NewRenderingVertexIndex < OldVertexBufferRenderingVertexCount )
+			{
+				if( OptionalCopyFromRenderingVertexID.IsSet() )
+				{
+					const uint32 CopyFromRenderingVertexIndex = OptionalCopyFromRenderingVertexID->GetValue();
+
+					// Copy from the specified vertex
+					StaticMeshLOD.PositionVertexBuffer.VertexPosition( NewRenderingVertexIndex ) = StaticMeshLOD.PositionVertexBuffer.VertexPosition( CopyFromRenderingVertexIndex );
+					StaticMeshLOD.VertexBuffer.SetVertexTangents(
+						NewRenderingVertexIndex,
+						StaticMeshLOD.VertexBuffer.VertexTangentX( CopyFromRenderingVertexIndex ),
+						StaticMeshLOD.VertexBuffer.VertexTangentY( CopyFromRenderingVertexIndex ),
+						StaticMeshLOD.VertexBuffer.VertexTangentZ( CopyFromRenderingVertexIndex ) );
+					for( int32 UVIndex = 0; UVIndex < NumUVs; ++UVIndex )
+					{
+						StaticMeshLOD.VertexBuffer.SetVertexUV( NewRenderingVertexIndex, UVIndex, StaticMeshLOD.VertexBuffer.GetVertexUV( CopyFromRenderingVertexIndex, UVIndex ) );
+					}
+
+					if( bHasColors )
+					{
+						StaticMeshLOD.ColorVertexBuffer.VertexColor( NewRenderingVertexIndex ) = StaticMeshLOD.ColorVertexBuffer.VertexColor( CopyFromRenderingVertexIndex );
+					}
+				}
+				else
+				{
+					// Initialize the new vertices to some defaults
+					StaticMeshLOD.PositionVertexBuffer.VertexPosition( NewRenderingVertexIndex ) = ReferencedVertex.VertexPosition;
+					StaticMeshLOD.VertexBuffer.SetVertexTangents(
+						NewRenderingVertexIndex,
+						FVector::ZeroVector,
+						FVector::ZeroVector,
+						FVector::ZeroVector );
+					for( int32 UVIndex = 0; UVIndex < NumUVs; ++UVIndex )
+					{
+						StaticMeshLOD.VertexBuffer.SetVertexUV( NewRenderingVertexIndex, UVIndex, FVector2D::ZeroVector );
+					}
+
+					if( bHasColors )
+					{
+						StaticMeshLOD.ColorVertexBuffer.VertexColor( NewRenderingVertexIndex ) = FColor::White;
+					}
+				}
+			}
+			else
+			{
+				const int32 AppendVertexNumber = NewRenderingVertexIndex - OldVertexBufferRenderingVertexCount;
+				check( AppendVertexNumber >= 0 && AppendVertexNumber < NumNewVertexBufferRenderingVertices );
+				FStaticMeshBuildVertex& RenderingVertexToAppend = RenderingVerticesToAppend[ AppendVertexNumber ];
+
+				if( OptionalCopyFromRenderingVertexID.IsSet() )
+				{
+					const uint32 CopyFromRenderingVertexIndex = OptionalCopyFromRenderingVertexID->GetValue();
+
+					// Copy from the specified vertex
+					RenderingVertexToAppend.Position = StaticMeshLOD.PositionVertexBuffer.VertexPosition( CopyFromRenderingVertexIndex );
+					RenderingVertexToAppend.TangentX = StaticMeshLOD.VertexBuffer.VertexTangentX( CopyFromRenderingVertexIndex );
+					RenderingVertexToAppend.TangentY = StaticMeshLOD.VertexBuffer.VertexTangentY( CopyFromRenderingVertexIndex );
+					RenderingVertexToAppend.TangentZ = StaticMeshLOD.VertexBuffer.VertexTangentZ( CopyFromRenderingVertexIndex );
+					for( int32 UVIndex = 0; UVIndex < NumUVs; ++UVIndex )
+					{
+						RenderingVertexToAppend.UVs[ UVIndex ] = StaticMeshLOD.VertexBuffer.GetVertexUV( CopyFromRenderingVertexIndex, UVIndex );
+					}
+					RenderingVertexToAppend.Color = bHasColors ? StaticMeshLOD.ColorVertexBuffer.VertexColor( CopyFromRenderingVertexIndex ) : FColor::White;
+				}
+				else
+				{
+					// Initialize the new vertices to some defaults
+					RenderingVertexToAppend.Position = ReferencedVertex.VertexPosition;
+					RenderingVertexToAppend.TangentX = FPackedNormal::ZeroNormal;
+					RenderingVertexToAppend.TangentY = FPackedNormal::ZeroNormal;
+					RenderingVertexToAppend.TangentZ = FPackedNormal::ZeroNormal;
+					for( int32 UVIndex = 0; UVIndex < NumUVs; ++UVIndex )
+					{
+						RenderingVertexToAppend.UVs[ UVIndex ] = FVector2D::ZeroVector;
+					}
+					RenderingVertexToAppend.Color = FColor::White;
+				}
+			}
+		}
+
+
+		if( RenderingVerticesToAppend.Num() > 0 )
+		{
+			StaticMeshLOD.VertexBuffer.AppendVertices( RenderingVerticesToAppend.GetData(), RenderingVerticesToAppend.Num() );
+			StaticMeshLOD.PositionVertexBuffer.AppendVertices( RenderingVerticesToAppend.GetData(), RenderingVerticesToAppend.Num() );
+			if( bHasColors )
+			{
+				StaticMeshLOD.ColorVertexBuffer.AppendVertices( RenderingVerticesToAppend.GetData(), RenderingVerticesToAppend.Num() );
+			}
+		}
+	}
+}
+
+
+void UEditableStaticMesh::DeleteOrphanRenderingVertices( const TArray<FRenderingVertexID>& RenderingVertexIDs )
+{
+	// Don't actually delete any vertices, but instead just mark them as unused.
+	for( const FRenderingVertexID RenderingVertexIDToDelete : RenderingVertexIDs )
+	{
+		RenderingVertices.RemoveAt( RenderingVertexIDToDelete.GetValue() );
+	}
+
+	// @todo mesheditor urgent perf: Add a 'compact' feature that actually deletes unused vertices
+	// This will invalidate all existing IDs though, so it needs to be handled pretty carefully.  Take a look
+	// at the history of this function in revision control to see an implementation that did this.
+	// Probably also, we'd want to do this for all other sparse arrays in the mesh too!
+}
+
+
+void UEditableStaticMesh::DeleteOrphanVertices_Internal( const TArray<FVertexID>& VertexIDsToDelete )
+{
+	static TArray<FRenderingVertexID> RenderingVertexIDsToDelete;
+	RenderingVertexIDsToDelete.Reset();
+
+	for( int32 VertexNumber = 0; VertexNumber < VertexIDsToDelete.Num(); ++VertexNumber )
+	{
+		const FVertexID VertexID = VertexIDsToDelete[ VertexNumber ];
+
+		const FEditableStaticMeshVertex& Vertex = Vertices[ VertexID.GetValue() ];
+
+		for( const FRenderingVertexID RenderingVertexID : Vertex.RenderingVertexIDs )
+		{
+			check( !RenderingVertexIDsToDelete.Contains( RenderingVertexID ) );	// Two vertices should never be sharing the same rendering vertex index
+			check( RenderingVertices.IsAllocated( RenderingVertexID.GetValue() ) );
+			RenderingVertexIDsToDelete.AddUnique( RenderingVertexID );
+		}
+
+		// Vertex must be orphaned before it is deleted!
+		check( Vertex.ConnectedEdgeIDs.Num() == 0 );
+
+		// Delete the vertex
+		Vertices.RemoveAt( VertexID.GetValue() );
+	}
+
+	// Delete the rendering vertices from the static mesh's vertex buffers
+	if( RenderingVertexIDsToDelete.Num() > 0 )
+	{
+		DeleteOrphanRenderingVertices( RenderingVertexIDsToDelete );
+	}
+}
+
+
+void UEditableStaticMesh::DeleteEdges_Internal( const TArray<FEdgeID>& EdgeIDsToDelete, const bool bDeleteOrphanedVertices )
+{
+	// Keep track of any vertices we orphaned, so we can delete them after we unhook everything
+	static TArray< FVertexID > OrphanedVertexIDs;
+	OrphanedVertexIDs.Reset();
+
+	for( const FEdgeID EdgeID : EdgeIDsToDelete )
+	{
+		FEditableStaticMeshEdge& Edge = Edges[ EdgeID.GetValue() ];
+
+		for( const FVertexID EdgeVertexID : Edge.VertexIDs )
+		{
+			FEditableStaticMeshVertex& Vertex = Vertices[ EdgeVertexID.GetValue() ];
+			const int32 NumRemovedEdges = Vertex.ConnectedEdgeIDs.RemoveSingle( EdgeID );
+			check( NumRemovedEdges == 1 );
+
+			// If the vertex has no more edges connected, we'll keep track of that so we can delete the vertex later
+			if( Vertex.ConnectedEdgeIDs.Num() == 0 )
+			{
+				check( !OrphanedVertexIDs.Contains( EdgeVertexID ) );	// Orphaned vertex shouldn't have already been orphaned by an earlier deleted edge passed into this function
+				OrphanedVertexIDs.Add( EdgeVertexID );
+			}
+		}
+
+		// Delete the edge
+		Edges.RemoveAt( EdgeID.GetValue() );
+	}
+
+
+	// If we orphaned any vertices and we were asked to delete those, then we'll go ahead and do that now.
+	if( bDeleteOrphanedVertices && OrphanedVertexIDs.Num() > 0 )
+	{
+		DeleteOrphanVertices( OrphanedVertexIDs );
+	}
+}
+
+
+void UEditableStaticMesh::DeletePolygon_Internal( const FPolygonRef PolygonRef, const bool bDeleteOrphanedEdges, const bool bDeleteOrphanedVertices, const bool bDeleteEmptySections )
+{
+	FEditableStaticMeshSection& Section = Sections[ PolygonRef.SectionID.GetValue() ];
+
+	// Keep track of any edges we orphaned, so we can delete them after we unhook everything
+	static TArray< FEdgeID > OrphanedEdgeIDs;
+	OrphanedEdgeIDs.Reset();	// @todo mesheditor perf: Possibly use inline allocators/alloca instead of statics all over the place (threading)
+
+
+	FEditableStaticMeshPolygon& Polygon = Section.Polygons[ PolygonRef.PolygonID.GetValue() ];
+
+	// Update all of our edges.  They no longer connect with us.
+	{
+		static TArray< FEdgeID > ContourEdgeIDs;
+		GetPolygonPerimeterEdges( PolygonRef, /* Out */ ContourEdgeIDs );
+		for( const FEdgeID EdgeID : ContourEdgeIDs )
+		{
+			FEditableStaticMeshEdge& Edge = Edges[ EdgeID.GetValue() ];
+
+			int32 ExistingPolygonNumber;
+			bool bWasFound = Edge.ConnectedPolygons.Find( PolygonRef, /* Out */ ExistingPolygonNumber );
+
+			check( bWasFound );
+			Edge.ConnectedPolygons.RemoveAt( ExistingPolygonNumber );
+
+			// If the edge has no more polygons connected, we'll keep track of that so we can delete the edge later
+			if( Edge.ConnectedPolygons.Num() == 0 )
+			{
+				OrphanedEdgeIDs.Add( EdgeID );
+			}
+		}
+
+		for( int32 HoleNumber = 0; HoleNumber < Polygon.HoleContours.Num(); ++HoleNumber )
+		{
+			GetPolygonHoleEdges( PolygonRef, HoleNumber, /* Out */ ContourEdgeIDs );
+			for( const FEdgeID EdgeID : ContourEdgeIDs )
+			{
+				FEditableStaticMeshEdge& Edge = Edges[ EdgeID.GetValue() ];
+
+				int32 ExistingPolygonNumber;
+				bool bWasFound = Edge.ConnectedPolygons.Find( PolygonRef, /* Out */ ExistingPolygonNumber );
+
+				check( bWasFound );
+				Edge.ConnectedPolygons.RemoveAt( ExistingPolygonNumber );
+
+				// If the edge has no more polygons connected, we'll keep track of that so we can delete the edge later
+				if( Edge.ConnectedPolygons.Num() == 0 )
+				{
+					OrphanedEdgeIDs.Add( EdgeID );
+				}
+			}
+		}
+	}
+
+	// Removes all of a polygon's triangles (including rendering triangles from the index buffer.)
+	DeletePolygonTriangles( PolygonRef );
+
+	// Delete our polygon's rendering vertices, as long as they aren't used by a different polygon.
+	// Imported meshes can share rendering vertices between polygons, if they have the same data (e.g. smooth edges and same UVs.)
+	{
+		// Delete the rendering vertices that are no longer used by any polygons in the mesh (if any)
+
+		// Check to see if we'll be orphaning any of the rendering vertices of this triangle.  We can only delete rendering
+		// vertices that are not actually used by other polygons.  Remember, any given vertex could have many rendering
+		// vertex copies (for discreet normals, etc.)
+		// @todo mesheditor holes: Also need to check holes?  Maybe?  Other polygons can't share those verts though...  At the least,
+		// all of our hole rendering vertices need to be treated as orphans and deleted!
+		static TArray< FRenderingVertexID > OrphanedRenderingVertexIDs;
+		OrphanedRenderingVertexIDs.Reset();
+
+		for( int32 PerimeterVertexNumber = 0; PerimeterVertexNumber < Polygon.PerimeterContour.Vertices.Num(); ++PerimeterVertexNumber )
+		{
+			const FVertexID PerimeterVertexID = Polygon.PerimeterContour.Vertices[ PerimeterVertexNumber ].VertexID;
+			const FRenderingVertexID PerimeterRenderingVertexID = Polygon.PerimeterContour.Vertices[ PerimeterVertexNumber ].RenderingVertexID;
+
+			// Assume this rendering vertex is an orphan unless we find a connected polygon that is also using it
+			bool bIsOrphan = true;
+			for( const FEdgeID ConnectedEdgeID : Vertices[ PerimeterVertexID.GetValue() ].ConnectedEdgeIDs )
+			{
+				for( const FPolygonRef OtherPolygonRef : Edges[ ConnectedEdgeID.GetValue() ].ConnectedPolygons )
+				{
+					// Ignore ourselves
+					if( OtherPolygonRef != PolygonRef )
+					{
+						const FEditableStaticMeshPolygon& OtherPolygon = Sections[ OtherPolygonRef.SectionID.GetValue() ].Polygons[ OtherPolygonRef.PolygonID.GetValue() ];
+
+						if( OtherPolygon.PerimeterContour.Vertices.ContainsByPredicate(
+							[ PerimeterRenderingVertexID ]( const FEditableStaticMeshPolygonContourVertex& V ) { return( PerimeterRenderingVertexID == V.RenderingVertexID ); }
+						))
+						{
+							bIsOrphan = false;
+							break;
+						}
+					}
+				}
+
+				if( !bIsOrphan )
+				{
+					break;
+				}
+			}
+
+			if( bIsOrphan )
+			{
+				FEditableStaticMeshVertex& PerimeterVertex = Vertices[ PerimeterVertexID.GetValue() ];
+				verify( PerimeterVertex.RenderingVertexIDs.RemoveSingleSwap( PerimeterRenderingVertexID ) == 1 );
+				OrphanedRenderingVertexIDs.Add( PerimeterRenderingVertexID );				
+			}
+		}
+
+		if( OrphanedRenderingVertexIDs.Num() )
+		{
+			DeleteOrphanRenderingVertices( OrphanedRenderingVertexIDs );
+		}
+	}
+
+
+	// Delete the polygon
+	Section.Polygons.RemoveAt( PolygonRef.PolygonID.GetValue() );
+
+	
+	// If we orphaned any edges and we were asked to delete those, then we'll go ahead and do that now.
+	// Deleting the edge may also delete orphaned vertices, if we were told to.
+	if( bDeleteOrphanedEdges && OrphanedEdgeIDs.Num() > 0 )
+	{
+		DeleteEdges( OrphanedEdgeIDs, bDeleteOrphanedVertices );
+	}
+
+	// If there are no longer any polygons left in the section, delete it too
+	if( bDeleteEmptySections && Section.Polygons.Num() == 0 )
+	{
+		DeleteSection( PolygonRef.SectionID );
+	}
+}
+
+
+void UEditableStaticMesh::DeletePolygonTriangles( const FPolygonRef PolygonRef )
+{
+	FEditableStaticMeshSection& Section = Sections[ PolygonRef.SectionID.GetValue() ];
+	FEditableStaticMeshPolygon& Polygon = Section.Polygons[ PolygonRef.PolygonID.GetValue() ];
+	const int32 NumTrianglesToRemove = Polygon.TriangulatedPolygonTriangleIndices.Num();
+	if( NumTrianglesToRemove > 0 )
+	{
+		// Kill the polygon's rendering triangles in the static mesh
+
+		// Remove all of the polygon's triangles from our editable mesh's triangle list.  While doing this, we'll keep
+		// track of all of the rendering mesh triangles that we'll need to remove later on.  We'll also figure out which
+		// rendering vertices will need to be removed from their corresponding vertex
+		for( const FTriangleID TriangleIndexToRemove : Polygon.TriangulatedPolygonTriangleIndices )
+		{
+			// Remove this triangle from our editable mesh
+			Section.Triangles.RemoveAt( TriangleIndexToRemove.GetValue() );
+		}
+																									  
+		// Update the index buffer by removing entries, and the rendering sections with new section counts
+		DeleteRenderingTrianglesForSectionTriangles( PolygonRef.SectionID, Polygon.TriangulatedPolygonTriangleIndices );
+		Polygon.TriangulatedPolygonTriangleIndices.Reset();
+	}
+}
+
+
+void UEditableStaticMesh::DeleteRenderingTrianglesForSectionTriangles( const FSectionID SectionID, const TArray<FTriangleID>& SectionTriangleIDsToRemove )
+{
+	if( !IsPreviewingSubdivisions() )
+	{
+		FStaticMeshLODResources& StaticMeshLOD = GetStaticMeshLOD();
+		const uint32 RenderingSectionIndex = Sections[ SectionID.GetValue() ].RenderingSectionIndex;
+		FStaticMeshSection& RenderingSection = StaticMeshLOD.Sections[ RenderingSectionIndex ];
+
+
+		const int32 NumTrianglesToRemove = SectionTriangleIDsToRemove.Num();
+		check( NumTrianglesToRemove > 0 );
+
+		// @todo mesheditor: We're simply changing existing triangles to be degenerates, so the section's total triangle count doesn't change.
+		// Later, when we compact data, we'll need to update the RenderingSection.NumTriangles (and RenderingSection.FirstIndex) for other
+		// sections, potentially.
+
+		// Update the index buffer by removing entries
+		{
+			// @todo mesheditor urgent: What about other index buffers in the mesh (DepthOnlyIndexBuffer, Wireframe, etc.)  We need to remove our triangles from those too!
+
+			for( const FTriangleID SectionTriangleIDToRemove : SectionTriangleIDsToRemove )
+			{
+				const uint32 RenderingTriangleFirstVertexIndex = FEditableStaticMeshSection::TriangleIndexToRenderingTriangleFirstIndex( RenderingSection, SectionTriangleIDToRemove );
+
+				// Make the indices degenerate.  We don't want to actually remove the indices from the index buffer, as that's can
+				// be a really slow operation.  The mesh can be compacted later on to free up the memory.
+				for( int32 TriangleVertexNumber = 0; TriangleVertexNumber < 3; ++TriangleVertexNumber )
+				{
+					// @todo mesheditor subdiv: 0 could be outside of our section's MinVertexIndex/MaxVertexIndex range.  We should choose the first valid
+					// rendering vertex index instead.
+					StaticMeshLOD.IndexBuffer.SetIndex( RenderingTriangleFirstVertexIndex + TriangleVertexNumber, 0 );
+				}
+			}
+		}
+	}
+}
+
+
+inline const FStaticMeshLODResources& UEditableStaticMesh::GetStaticMeshLOD() const
+{
+	const FStaticMeshRenderData& StaticMeshRenderData = *StaticMesh->RenderData;
+	const FStaticMeshLODResources& StaticMeshLOD = StaticMeshRenderData.LODResources[ SubMeshAddress.LODIndex ];
+	return StaticMeshLOD;
+}
+
+
+FStaticMeshLODResources& UEditableStaticMesh::GetStaticMeshLOD()
+{
+	FStaticMeshRenderData& StaticMeshRenderData = *StaticMesh->RenderData;
+	FStaticMeshLODResources& StaticMeshLOD = StaticMeshRenderData.LODResources[ SubMeshAddress.LODIndex ];
+	return StaticMeshLOD;
+}
+
+
+bool UEditableStaticMesh::DoesPolygonPerimeterVertexHaveDiscreetRenderingVertex( const FPolygonRef PolygonRef, const int32 PolygonVertexNumber ) const
+{
+	const FRenderingVertexID RenderingVertexID = GetPolygonPerimeterRenderingVertex( PolygonRef, PolygonVertexNumber );
+
+	bool bHasDiscreetRenderingVertex = true;
+	{
+		const FEditableStaticMeshVertex& Vertex = Vertices[ GetPolygonPerimeterVertex( PolygonRef, PolygonVertexNumber ).GetValue() ];
+		for( int32 EdgeNumber = 0; EdgeNumber < Vertex.ConnectedEdgeIDs.Num(); ++EdgeNumber )
+		{
+			const FEdgeID EdgeID = Vertex.ConnectedEdgeIDs[ EdgeNumber ];
+			const FEditableStaticMeshEdge& Edge = Edges[ EdgeID.GetValue() ];
+
+			for( const FPolygonRef& ConnectedPolygonRef : Edge.ConnectedPolygons )
+			{
+				// Ignore ourselves for this check.
+				if( ConnectedPolygonRef != PolygonRef )
+				{
+					const FEditableStaticMeshPolygon& OtherPolygon = Sections[ ConnectedPolygonRef.SectionID.GetValue() ].Polygons[ ConnectedPolygonRef.PolygonID.GetValue() ];
+					if( OtherPolygon.PerimeterContour.Vertices.ContainsByPredicate(
+						[ RenderingVertexID ]( const FEditableStaticMeshPolygonContourVertex& V )
+						{
+							return( V.RenderingVertexID == RenderingVertexID );
+						}
+					))
+					{
+						// Oh no -- a different polygon is referencing our rendering vertex.  We'll need to make a new one.
+						bHasDiscreetRenderingVertex = false;
+						break;
+					}
+				}
+			}
+
+			if( !bHasDiscreetRenderingVertex )
+			{
+				break;
+			}
+		}
+	}
+
+	// @todo mesheditor debug
+	// UE_LOG( LogMeshEditingRuntime, Log, TEXT( "bHasDiscreetRenderingVertex=%i for Section=%i, Polygon=%i, PolyVertex=%i" ), bHasDiscreetRenderingVertex ? 1 : 0, PolygonRef.SectionID.GetIDValue(), PolygonRef.PolygonID.GetIDValue(), PolygonVertexNumber );
+	return bHasDiscreetRenderingVertex;
+}
+
+
+FRenderingVertexID UEditableStaticMesh::MakeDiscreetPolygonPerimeterRenderingVertexIfNeeded( const FPolygonRef PolygonRef, const int32 PolygonVertexNumber )
+{
+	const FVertexID VertexID = GetPolygonPerimeterVertex( PolygonRef, PolygonVertexNumber );
+	FEditableStaticMeshVertex& Vertex = Vertices[ VertexID.GetValue() ];
+
+	FRenderingVertexID DiscreetRenderingVertexID = GetPolygonPerimeterRenderingVertex( PolygonRef, PolygonVertexNumber );
+
+	// Check to see which polygons are using this vertex.  We need to make sure we have a unique rendering vertex for
+	// polygon vertex, because we don't want to affect other polygons with this change!
+	// @todo mesheditor urgent hole: Also need to do this for HOLES.  HOLES really need a refactor...
+	if( !DoesPolygonPerimeterVertexHaveDiscreetRenderingVertex( PolygonRef, PolygonVertexNumber ) )
+	{
+		// Copy per-triangle vertex data (e.g. polygon perimeter vertex attributes) from the existing rendering vertex, so that we
+		// don't have to bother setting those manually after cloning the vertex.
+		const FRenderingVertexID CopyFromRenderingVertexID = DiscreetRenderingVertexID;
+		const FRenderingVertexID NewRenderingVertexID = AddNewRenderingVertexToPolygonPerimeter( PolygonRef, PolygonVertexNumber, CopyFromRenderingVertexID );
+
+		// Update our triangle index buffer.  We need to point to our new vertex.
+		// @todo mesheditor: We only really need to update indices if the rendering vertex could have been shared between other triangles in the mesh.  Currently this only happens with the initial rendering vertices from a converted static mesh.  We currently always use discreet rendering vertices for each triangle while editing.
+		FEditableStaticMeshSection& Section = Sections[ PolygonRef.SectionID.GetValue() ];
+		FEditableStaticMeshPolygon& Polygon = Section.Polygons[ PolygonRef.PolygonID.GetValue() ];
+
+		FStaticMeshLODResources& StaticMeshLOD = GetStaticMeshLOD();
+
+		const uint32 RenderingSectionIndex = Section.RenderingSectionIndex;
+		FStaticMeshSection& RenderingSection = StaticMeshLOD.Sections[ RenderingSectionIndex ];
+
+		// For each of our triangulated triangles
+		for( const FTriangleID TriangleID : Polygon.TriangulatedPolygonTriangleIndices )
+		{
+			for( int32 VertexNumber = 0; VertexNumber < 3; ++VertexNumber )
+			{
+				if( Section.Triangles[ TriangleID.GetValue() ].RenderingVertexIDs[ VertexNumber ] == DiscreetRenderingVertexID )
+				{
+					// Update the triangle rendering vertex ID to the newly added rendering vertex
+					Section.Triangles[ TriangleID.GetValue() ].RenderingVertexIDs[ VertexNumber ] = NewRenderingVertexID;
+
+					if( !IsPreviewingSubdivisions() )
+					{
+						const uint32 RenderingTriangleFirstIndex = FEditableStaticMeshSection::TriangleIndexToRenderingTriangleFirstIndex( RenderingSection, TriangleID );
+						const uint32 IndexOfIndex = RenderingTriangleFirstIndex + VertexNumber;
+						check( StaticMeshLOD.IndexBuffer.GetIndex( IndexOfIndex ) == DiscreetRenderingVertexID.GetValue() );
+
+						UpdateIndexBufferFormatIfNeeded( NewRenderingVertexID );
+
+						StaticMeshLOD.IndexBuffer.SetIndex( IndexOfIndex, NewRenderingVertexID.GetValue() );
+					}
+				}
+			}
+		}
+
+		DiscreetRenderingVertexID = NewRenderingVertexID;
+
+		// @todo mesheditor debug
+		//UE_LOG( LogMeshEditingRuntime, Log, TEXT( "Made DiscreetRenderingVertex=%i for Section=%i, Polygon=%i, PolyVertex=%i" ), DiscreetRenderingVertexID.GetValue(), PolygonRef.SectionID.GetValue(), PolygonRef.PolygonID.GetValue(), PolygonVertexNumber );
+	}
+
+	return DiscreetRenderingVertexID;
+}
+
+
+FRenderingVertexID UEditableStaticMesh::AddNewRenderingVertexToPolygonPerimeter( const FPolygonRef PolygonRef, const int32 PolygonVertexNumber, const TOptional<FRenderingVertexID> OptionalCopyFromRenderingVertexID )
+{
+	const FVertexID VertexID = GetPolygonPerimeterVertex( PolygonRef, PolygonVertexNumber );
+
+	static TArray<FVertexID> VertexIDs;
+	VertexIDs.Reset();
+	VertexIDs.Add( VertexID );
+
+	// OK, so we need a unique rendering vertex.  Let's make one now.  The vertex data will start "zeroed out".
+	static TArray<FRenderingVertexID> NewRenderingVertexIDs;
+	CreateRenderingVertices( VertexIDs, OptionalCopyFromRenderingVertexID, /* Out */ NewRenderingVertexIDs );
+
+	// Update our polygon
+	FEditableStaticMeshSection& Section = Sections[ PolygonRef.SectionID.GetValue() ];
+	FEditableStaticMeshPolygon& Polygon = Section.Polygons[ PolygonRef.PolygonID.GetValue() ];
+	Polygon.PerimeterContour.Vertices[ PolygonVertexNumber ].RenderingVertexID = NewRenderingVertexIDs[ 0 ];	// @todo mesheditor perf: This should really do multiple vertices at once (the issue stems from SetPolygonsVertexAttributes only doing one vertex at a time)
+
+	return NewRenderingVertexIDs[ 0 ];
+}
+
+
+void UEditableStaticMesh::SetEdgeVertices_Internal( const FEdgeID EdgeID, const FVertexID NewVertexID0, const FVertexID NewVertexID1 )
+{
+	FEditableStaticMeshEdge& Edge = Edges[ EdgeID.GetValue() ];
+
+	for( uint32 EdgeVertexNumber = 0; EdgeVertexNumber < 2; ++EdgeVertexNumber )
+	{
+		// Disconnect the edge from it's existing vertices
+		const FVertexID VertexID = Edge.VertexIDs[ EdgeVertexNumber ];
+		FEditableStaticMeshVertex& Vertex = Vertices[ VertexID.GetValue() ];
+		verify( Vertex.ConnectedEdgeIDs.RemoveSingleSwap( EdgeID ) == 1 );	// Must have been already connected!
+	}
+
+	Edge.VertexIDs[0] = NewVertexID0;
+	Edge.VertexIDs[1] = NewVertexID1;
+
+	// Connect the new vertices to the edge
+	for( uint32 EdgeVertexNumber = 0; EdgeVertexNumber < 2; ++EdgeVertexNumber )
+	{
+		const FVertexID VertexID = Edge.VertexIDs[ EdgeVertexNumber ];
+		FEditableStaticMeshVertex& Vertex = Vertices[ VertexID.GetValue() ];
+
+		check( !Vertex.ConnectedEdgeIDs.Contains( EdgeID ) );	// Should not have already been connected
+		Vertex.ConnectedEdgeIDs.Add( EdgeID );
+	}
+}
+
+
+void UEditableStaticMesh::InsertPolygonPerimeterVertices_Internal( const FPolygonRef PolygonRef, const int32 InsertBeforeVertexNumber, const TArray<FVertexAndAttributes>& VerticesToInsert )
+{
+	FEditableStaticMeshPolygon& Polygon = Sections[ PolygonRef.SectionID.GetValue() ].Polygons[ PolygonRef.PolygonID.GetValue() ];
+
+	for( int32 InsertVertexIter = 0; InsertVertexIter < VerticesToInsert.Num(); ++InsertVertexIter )
+	{
+		Polygon.PerimeterContour.Vertices.Insert( FEditableStaticMeshPolygonContourVertex(), InsertBeforeVertexNumber + InsertVertexIter );
+		const FVertexAndAttributes& VertexToInsert = VerticesToInsert[ InsertVertexIter ];
+
+		const int32 PolygonVertexNumber = InsertBeforeVertexNumber + InsertVertexIter;
+		Polygon.PerimeterContour.Vertices[ PolygonVertexNumber ].VertexID = VertexToInsert.VertexID;
+
+		const FRenderingVertexID RenderingVertexID = AddNewRenderingVertexToPolygonPerimeter( PolygonRef, PolygonVertexNumber );
+
+		for( int32 AttributeNumber = 0; AttributeNumber < VertexToInsert.PolygonVertexAttributes.Attributes.Num(); ++AttributeNumber )
+		{
+			const FMeshElementAttributeData& MeshElementAttribute = VertexToInsert.PolygonVertexAttributes.Attributes[ AttributeNumber ];
+			SetPolygonPerimeterVertexAttribute_Internal( PolygonRef, InsertBeforeVertexNumber + InsertVertexIter, MeshElementAttribute.AttributeName, MeshElementAttribute.AttributeIndex, MeshElementAttribute.AttributeValue );
+		}
+	}
+}
+
+
+void UEditableStaticMesh::RemovePolygonPerimeterVertices_Internal( const FPolygonRef PolygonRef, const int32 FirstVertexNumberToRemove, const int32 NumVerticesToRemove )
+{
+	FEditableStaticMeshPolygon& Polygon = Sections[ PolygonRef.SectionID.GetValue() ].Polygons[ PolygonRef.PolygonID.GetValue() ];
+
+	// @todo mesheditor: We're assuming these are all orphans because this function is only ever used to undo the addition of
+	// brand new vertices to existing polygons.  If we ever start sharing rendering rendering vertices between mesh vertices
+	// with our mesh editing operations (not just imported meshes, like we do now), we'll need to do an orphan check here first.
+	static TArray<FRenderingVertexID> OrphanedRenderingVertexIDs;
+	OrphanedRenderingVertexIDs.Reset();
+	for( int32 VertexIter = 0; VertexIter < NumVerticesToRemove; ++VertexIter )
+	{
+		const int32 PerimeterVertexNumber = FirstVertexNumberToRemove + VertexIter;
+		const FRenderingVertexID RenderingVertexID = Polygon.PerimeterContour.Vertices[ PerimeterVertexNumber ].RenderingVertexID;
+		FEditableStaticMeshVertex& PerimeterVertex = Vertices[ Polygon.PerimeterContour.Vertices[ PerimeterVertexNumber ].VertexID.GetValue() ];
+		verify( PerimeterVertex.RenderingVertexIDs.RemoveSingleSwap( RenderingVertexID ) == 1 );
+		OrphanedRenderingVertexIDs.Add( RenderingVertexID );
+	}
+
+	DeleteOrphanRenderingVertices( OrphanedRenderingVertexIDs );
+
+	Polygon.PerimeterContour.Vertices.RemoveAt( FirstVertexNumberToRemove, NumVerticesToRemove );
+}
+
+
+FSectionID UEditableStaticMesh::GetSectionIDFromMaterial_Internal( UMaterialInterface* Material, bool bCreateNewSectionIfNotFound )
+{
+	check(StaticMesh);
+	// Iterate through the sections sparse array looking for an entry whose material index matches.
+	for( TSparseArray<FEditableStaticMeshSection>::TConstIterator It( Sections ); It; ++It )
+	{
+		const FEditableStaticMeshSection& Section = *It;
+
+		if( StaticMesh->GetMaterial( Section.MaterialIndex ) == Material )
+		{
+			return FSectionID( It.GetIndex() );
+		}
+	}
+
+	// If we got here, the material index does not yet have a matching section.
+	if( bCreateNewSectionIfNotFound )
+	{
+		// @todo mesheditor: Currently new materials are added to the static mesh and never removed.
+		// They should be pruned every now and then, maybe when saving.
+		FSectionToCreate SectionToCreate;
+		SectionToCreate.Material = Material;
+		SectionToCreate.bEnableCollision = true;
+		SectionToCreate.bCastShadow = true;
+
+		return CreateSection( SectionToCreate );
+	}
+
+	return FSectionID::Invalid;
+}
+
+
+FSectionID UEditableStaticMesh::CreateSection_Internal( const FSectionToCreate& SectionToCreate )
+{
+	const int32 MaterialIndex = StaticMesh->StaticMaterials.AddUnique( FStaticMaterial( SectionToCreate.Material ) );
+
+	uint32 LODSectionIndex = 0;
+	if( !IsPreviewingSubdivisions() )
+	{
+		// Need to create a new rendering section. This is added to the end of the array.
+		// We also create a corresponding editable mesh section.
+		FStaticMeshLODResources& StaticMeshLOD = GetStaticMeshLOD();
+		if( SectionToCreate.OriginalRenderingSectionIndex == INDEX_NONE )
+		{
+			// Add a new rendering section to the end if a specific index was not requested
+			LODSectionIndex = StaticMeshLOD.Sections.Emplace();
+		}
+		else
+		{
+			// Otherwise add the rendering section at the specific index requested
+			LODSectionIndex = SectionToCreate.OriginalRenderingSectionIndex;
+			StaticMeshLOD.Sections.EmplaceAt( LODSectionIndex );
+
+			// Adjust rendering indices held by sections: any index above the one we just added now needs to be incremented.
+			for( FEditableStaticMeshSection& Section : Sections )
+			{
+				if( Section.RenderingSectionIndex > LODSectionIndex )
+				{
+					Section.RenderingSectionIndex++;
+				}
+			}
+		}
+
+		FStaticMeshSection& StaticMeshSection = StaticMeshLOD.Sections[ LODSectionIndex ];
+
+		// Initially the section is empty, and it occupies zero elements in the index buffer.
+		// It is still placed in the correct location within the index buffer, immediately following the previous section,
+		// as it is a requirement that consecutive sections are placed contiguously in the index buffer.
+		// Determine the first index based on the index range of the previous rendering section.
+		if( LODSectionIndex == 0 )
+		{
+			StaticMeshSection.FirstIndex = 0;
+		}
+		else
+		{
+			FStaticMeshSection& PreviousStaticMeshSection = StaticMeshLOD.Sections[ LODSectionIndex - 1 ];
+			const FSectionID SectionID = GetSectionForRenderingSectionIndex( LODSectionIndex - 1 );
+			check( SectionID != FSectionID::Invalid );
+			StaticMeshSection.FirstIndex = PreviousStaticMeshSection.FirstIndex + Sections[ SectionID.GetValue() ].MaxTriangles * 3;
+		}
+
+		// Fill in the remaining rendering section properties.
+		StaticMeshSection.NumTriangles = 0;
+		StaticMeshSection.MinVertexIndex = 0;
+		StaticMeshSection.MaxVertexIndex = 0;
+		StaticMeshSection.bEnableCollision = SectionToCreate.bEnableCollision;
+		StaticMeshSection.bCastShadow = SectionToCreate.bCastShadow;
+		StaticMeshSection.MaterialIndex = MaterialIndex;
+	}
+
+	// Copy this information into the editable mesh section (which is authoritative)
+	int32 SectionIndex;
+	if( SectionToCreate.OriginalSectionID == FSectionID::Invalid )
+	{
+		// Add new section to the next free slot if a concrete index was not specified
+		SectionIndex = Sections.Add( FEditableStaticMeshSection() );
+	}
+	else
+	{
+		// If a concrete index was specified (should only be on undo), try to insert a new section there
+		SectionIndex = SectionToCreate.OriginalSectionID.GetValue();
+		check( !Sections.IsAllocated( SectionIndex ) );
+		Sections.Insert( SectionIndex, FEditableStaticMeshSection() );
+	}
+
+	// Fill out the authoratitive section data
+	FEditableStaticMeshSection& Section = Sections[ SectionIndex ];
+	Section.RenderingSectionIndex = LODSectionIndex;
+	Section.MaterialIndex = MaterialIndex;
+	Section.bEnableCollision = SectionToCreate.bEnableCollision;
+	Section.bCastShadow = SectionToCreate.bCastShadow;
+	Section.MaxTriangles = 0;
+
+	const FSectionID SectionID( SectionIndex );
+
+	// Allow operation to be undone
+	FDeleteSectionChangeInput DeleteSectionChangeInput;
+	DeleteSectionChangeInput.SectionID  = SectionID;
+
+	AddUndo( MakeUnique<FDeleteSectionChange>( DeleteSectionChangeInput ) );
+
+	return SectionID;
+}
+
+
+void UEditableStaticMesh::DeleteSection_Internal( const FSectionID SectionID )
+{
+	FEditableStaticMeshSection& Section = Sections[ SectionID.GetValue() ];
+
+	// Prepare the change input struct
+	FSectionToCreate SectionToCreate;
+	SectionToCreate.Material = StaticMesh->StaticMaterials[ Section.MaterialIndex ].MaterialInterface;
+	SectionToCreate.bEnableCollision = Section.bEnableCollision;
+	SectionToCreate.bCastShadow = Section.bCastShadow;
+	SectionToCreate.OriginalSectionID = SectionID;
+
+	FCreateSectionChangeInput CreateSectionChangeInput;
+	CreateSectionChangeInput.SectionToCreate = SectionToCreate;
+
+	// Remove material slot associated with section
+	// @todo mesheditor: can more than one section share a material? Mesh editor currently assumes not, but this will break anything which does
+	const int32 MaterialIndex = Sections[ SectionID.GetValue() ].MaterialIndex;
+	StaticMesh->StaticMaterials.RemoveAt( MaterialIndex );
+
+	for( FEditableStaticMeshSection& SectionToAdjust : Sections )
+	{
+		if( SectionToAdjust.MaterialIndex > MaterialIndex )
+		{
+			SectionToAdjust.MaterialIndex--;
+		}
+	}
+
+	// Adjust rendering indices held by sections: any index above the one we just deleted now needs to be decremented.
+	const uint32 RenderingSectionIndex = Section.RenderingSectionIndex;
+
+	for( FEditableStaticMeshSection& SectionToAdjust : Sections )
+	{
+		if( SectionToAdjust.RenderingSectionIndex > RenderingSectionIndex )
+		{
+			SectionToAdjust.RenderingSectionIndex--;
+		}
+	}
+
+	if( !IsPreviewingSubdivisions() )
+	{
+		FStaticMeshLODResources& StaticMeshLOD = GetStaticMeshLOD();
+		const uint32 FirstIndex = StaticMeshLOD.Sections[ RenderingSectionIndex ].FirstIndex;
+
+		// Get current number of triangles allocated for this section
+		const int32 MaxTriangles = Section.MaxTriangles;
+
+		// Remove indices from this poisition in the index buffer
+		StaticMeshLOD.IndexBuffer.RemoveIndicesAt( FirstIndex, MaxTriangles * 3 );
+
+		// Adjust first index for all subsequent render sections to account for the indices just removed.
+		// It is guaranteed that index buffer indices are in the same order as the rendering sections.
+		const uint32 NumRenderingSections = StaticMeshLOD.Sections.Num();
+		for( uint32 Index = RenderingSectionIndex + 1; Index < NumRenderingSections; ++Index )
+		{
+			check( StaticMeshLOD.Sections[ Index ].FirstIndex >= FirstIndex );
+			StaticMeshLOD.Sections[ Index ].FirstIndex -= MaxTriangles * 3;
+		}
+
+		for( uint32 Index = 0; Index < NumRenderingSections; ++Index )
+		{
+			FStaticMeshSection& StaticMeshSection = StaticMeshLOD.Sections[ Index ];
+			if( StaticMeshSection.MaterialIndex > MaterialIndex )
+			{
+				StaticMeshSection.MaterialIndex--;
+			}
+		}
+
+		StaticMeshLOD.Sections.RemoveAt( RenderingSectionIndex );
+	}
+
+	// Remove the section from the sparse array
+	Sections.RemoveAt( SectionID.GetValue() );
+
+	AddUndo( MakeUnique<FCreateSectionChange>( CreateSectionChangeInput ) );
+}
+
+
+FSectionID UEditableStaticMesh::GetSectionForRenderingSectionIndex( const int32 RenderingSectionIndex ) const
+{
+	for( auto It = Sections.CreateConstIterator(); It; ++It )
+	{
+		if( It->RenderingSectionIndex == RenderingSectionIndex )
+		{
+			return FSectionID( It.GetIndex() );
+		}
+	}
+
+	return FSectionID::Invalid;
+}
