Branch snapshot for CL 3327560

diff --git a/Engine/Source/Runtime/MeshEditingRuntime/EditableStaticMesh.h b/Engine/Source/Runtime/MeshEditingRuntime/EditableStaticMesh.h
new file mode 100644
--- /dev/null
+++ b/Engine/Source/Runtime/MeshEditingRuntime/EditableStaticMesh.h
@@ -0,0 +1,628 @@
+// Copyright 1998-2015 Epic Games, Inc. All Rights Reserved.
+
+#pragma once
+
+#include "CoreMinimal.h"
+#include "EditableMesh.h"
+#include "EditableMeshCustomVersion.h"
+#include "StaticMeshResources.h"
+#include "EditableStaticMesh.generated.h"
+
+
+/**
+* Perform custom serialization for TSparseArray.
+* The default TSparseArray serialization also compacts all the elements, removing the gaps and changing the indices.
+* The indices are significant in editable meshes, hence this is a custom serializer which preserves them.
+*/
+template <typename T>
+void SerializeSparseArray( FArchive& Ar, TSparseArray<T>& Array )
+{
+	if( Ar.CustomVer( FEditableMeshCustomVersion::GUID ) < FEditableMeshCustomVersion::CustomSparseArraySerialization )
+	{
+		Ar << Array;
+	}
+	else
+	{
+		Array.CountBytes( Ar );
+
+		if( Ar.IsLoading() )
+		{
+			// Load array
+			TBitArray<> AllocatedIndices;
+			Ar << AllocatedIndices;
+
+			Array.Empty( AllocatedIndices.Num() );
+			for( auto It = TConstSetBitIterator<>( AllocatedIndices ); It; ++It )
+			{
+				Array.Insert( It.GetIndex(), T() );
+				Ar << Array[ It.GetIndex() ];
+			}
+		}
+		else
+		{
+			// Save array
+			const int32 MaxIndex = Array.GetMaxIndex();
+
+			// We have to build the TBitArray representing allocated indices by hand, as we don't have access to it from outside TSparseArray.
+			// @todo core: consider replacing TSparseArray serialization with this format.
+			TBitArray<> AllocatedIndices( false, MaxIndex );
+			for( int32 Index = 0; Index < MaxIndex; ++Index )
+			{
+				if( Array.IsAllocated( Index ) )
+				{
+					AllocatedIndices[ Index ] = true;
+				}
+			}
+			Ar << AllocatedIndices;
+
+			for( auto It = Array.CreateIterator(); It; ++It )
+			{
+				Ar << *It;
+			}
+		}
+	}
+}
+
+
+USTRUCT()
+struct FRenderingVertexID : public FElementID
+{
+	GENERATED_BODY()
+
+	FRenderingVertexID()
+	{
+	}
+
+	explicit FRenderingVertexID( const FElementID InitElementID )
+		: FElementID( InitElementID.GetValue() )
+	{
+	}
+
+	explicit FRenderingVertexID( const uint32 InitIDValue )
+		: FElementID( InitIDValue )
+	{
+	}
+
+	/** Invalid rendering vertex ID */
+	MESHEDITINGRUNTIME_API static const FRenderingVertexID Invalid;
+};
+
+
+USTRUCT()
+struct FEditableStaticMeshVertex
+{
+	GENERATED_BODY()
+
+	/** Position of the vertex.  This is also stored in every rendering vertex in the actual static mesh, but
+	    we need a copy here because the vertex might not have any rendering vertices (in the case where no
+		triangles are connected to it.) */
+	UPROPERTY()
+	FVector VertexPosition;
+
+	/** All of the extra rendering vertices generated for this editable mesh vertex.  A rendering mesh may have
+	    multiple vertices that represent a single editable mesh vertex position (for discreet normals, etc.) */
+	UPROPERTY()
+	TArray< FRenderingVertexID > RenderingVertexIDs;
+
+	/** The edges connected to this vertex */
+	UPROPERTY()
+	TArray< FEdgeID > ConnectedEdgeIDs;
+
+	/** When subdivisions are enabled, this controls how sharp the vertex is, between 0.0 and 1.0. */
+	UPROPERTY()
+	float CornerSharpness;	// @todo mesheditor subdiv: Not really used by static meshes at all.  Only for UEditableMeshes that use subdivision features.  Move elsewhere.?
+
+	/** Serializer */
+	friend FArchive& operator<<( FArchive& Ar, FEditableStaticMeshVertex& Vertex )
+	{
+		Ar << Vertex.VertexPosition;
+		Ar << Vertex.RenderingVertexIDs;
+		Ar << Vertex.ConnectedEdgeIDs;
+		Ar << Vertex.CornerSharpness;
+		return Ar;
+	}
+};
+
+
+USTRUCT()
+struct FEditableStaticMeshRenderingVertex
+{
+	GENERATED_BODY()
+
+	UPROPERTY()
+	FVertexID VertexID;
+
+	/** Serializer */
+	friend FArchive& operator<<( FArchive& Ar, FEditableStaticMeshRenderingVertex& Vertex )
+	{
+		Ar << Vertex.VertexID;
+		return Ar;
+	}
+};
+
+
+USTRUCT()
+struct FEditableStaticMeshEdge
+{
+	GENERATED_BODY()
+
+	/** IDs of the two editable mesh vertices that make up this edge.  The winding direction is not defined. */
+	UPROPERTY()
+	FVertexID VertexIDs[ 2 ];
+
+	/** The polygons that share this edge.  It's best if there are always only two polygons that share
+	    the edge, and those polygons are facing the same direction */
+	UPROPERTY()
+	TArray< FPolygonRef > ConnectedPolygons;
+
+	/** Whether this edge is 'hard' or not, for the purpose of vertex normal and tangent generation */
+	UPROPERTY()
+	bool bIsHardEdge;
+
+	/** When subdivisions are enabled, this controls how sharp the creasing of this edge will be, between 0.0 and 1.0. */
+	UPROPERTY()
+	float CreaseSharpness;
+
+	/** Serializer */
+	friend FArchive& operator<<( FArchive& Ar, FEditableStaticMeshEdge& Edge )
+	{
+		Ar << Edge.VertexIDs[ 0 ];
+		Ar << Edge.VertexIDs[ 1 ];
+		Ar << Edge.ConnectedPolygons;
+		Ar << Edge.bIsHardEdge;
+		Ar << Edge.CreaseSharpness;
+		return Ar;
+	}
+};
+
+
+USTRUCT()
+struct FEditableStaticMeshTriangle
+{
+	GENERATED_BODY()
+
+	/** The three vertices that make up this triangle.  The winding direction is not defined.  Note that the order of these 
+	    are the same as the order of the triangles in the static mesh index buffer.  They must always be kept in sync!  To find
+		the first rendering triangle vertex index in the index buffer, multiply the triangle index in this array
+		by three (three vertices per triangle), then add the rendering section's start index */
+	UPROPERTY()
+	FRenderingVertexID RenderingVertexIDs[ 3 ];
+
+	/** Serializer */
+	friend FArchive& operator<<( FArchive& Ar, FEditableStaticMeshTriangle& Tri )
+	{
+		Ar << Tri.RenderingVertexIDs[ 0 ];
+		Ar << Tri.RenderingVertexIDs[ 1 ];
+		Ar << Tri.RenderingVertexIDs[ 2 ];
+		return Ar;
+	}
+};
+
+
+USTRUCT()
+struct FEditableStaticMeshPolygonContourVertex
+{
+	GENERATED_BODY()
+
+	/** Vertex ID representing this vertex */
+	UPROPERTY()
+	FVertexID VertexID;
+
+	/** Rendering vertex ID representing this vertex */
+	UPROPERTY()
+	FRenderingVertexID RenderingVertexID;
+
+	/** Other per-vertex polygon contour attributes follow. First UVs. */
+	UPROPERTY()
+	TArray< FVector2D > VertexUVs;
+
+	/** Normal vector */
+	UPROPERTY()
+	FVector Normal;
+	
+	/** Tangent vector */
+	UPROPERTY()
+	FVector Tangent;
+
+	/** Basis determinant sign used to calculate the sense of the binormal */
+	UPROPERTY()
+	float BinormalSign;
+
+	/** Vertex color */
+	UPROPERTY()
+	FLinearColor Color;
+
+	/** Serializer */
+	friend FArchive& operator<<( FArchive& Ar, FEditableStaticMeshPolygonContourVertex& Vertex )
+	{
+		Ar << Vertex.VertexID;
+		Ar << Vertex.RenderingVertexID;	// @todo mesheditor serialization: Should not need to be serialized if we triangulate after load
+		Ar << Vertex.VertexUVs;
+		Ar << Vertex.Normal;	// @todo mesheditor serialization: Need a way to specify if a normal has been overridden and hence needs custom serialization instead of being generated
+		Ar << Vertex.Tangent;	// @todo mesheditor serialization: Should not need to be serialized if we regenerate normals and tangents after load
+		Ar << Vertex.BinormalSign;	// @todo mesheditor serialization: Should not need to be serialized if we regenerate normals and tangents after load
+
+		if( Ar.CustomVer( FEditableMeshCustomVersion::GUID ) >= FEditableMeshCustomVersion::WithVertexColors )
+		{
+			Ar << Vertex.Color;
+		}
+
+		return Ar;
+	}
+};
+
+
+USTRUCT()
+struct FEditableStaticMeshPolygonContour
+{
+	GENERATED_BODY()
+
+	/** The ordered list of vertices that make up the polygon contour. The winding direction is counter-clockwise. */
+	UPROPERTY()
+	TArray< FEditableStaticMeshPolygonContourVertex > Vertices;
+
+	/** Serializer */
+	friend FArchive& operator<<( FArchive& Ar, FEditableStaticMeshPolygonContour& Contour )
+	{
+		Ar << Contour.Vertices;
+		return Ar;
+	}
+};
+
+
+USTRUCT()
+struct FEditableStaticMeshPolygon
+{
+	GENERATED_BODY()
+
+	/** The outer boundary edges of this polygon */
+	UPROPERTY()
+	FEditableStaticMeshPolygonContour PerimeterContour;
+
+	/** Optional inner contours of this polygon that define holes inside of the polygon.  For the geometry to
+	    be considered valid, the hole contours should reside within the boundary of the polygon perimeter contour, 
+		and must not overlap each other.  No "nesting" of polygons inside the holes is supported -- those are 
+		simply separate polygons */
+	UPROPERTY()
+	TArray< FEditableStaticMeshPolygonContour > HoleContours;
+
+	/** Static meshes currently only support triangles.  We'll always triangulate polygons and keep track
+	    of all of the triangles here */
+	// @todo mesheditor: For other mesh formats, we may not need to assume triangles.  For example, modern
+	//		GPUs can render quads.  We still need to consider collision geometry though.
+	// @todo mesheditor: This should probably always be a transient thing, right?  We can re-triangulate on 
+	//		load?  Not sure.  We might want some edit-time control over triangulation, which should be saved/loaded.
+	UPROPERTY()
+	TArray< FTriangleID > TriangulatedPolygonTriangleIndices;
+
+	/** Serializer */
+	friend FArchive& operator<<( FArchive& Ar, FEditableStaticMeshPolygon& Polygon )
+	{
+		Ar << Polygon.PerimeterContour;
+		Ar << Polygon.HoleContours;
+		Ar << Polygon.TriangulatedPolygonTriangleIndices;
+		return Ar;
+	}
+};
+
+
+USTRUCT()
+struct FEditableStaticMeshSection
+{
+	GENERATED_BODY()
+
+	/** The rendering section index for this mesh section */
+	UPROPERTY()
+	uint32 RenderingSectionIndex;
+
+	/** The material index for this mesh section */
+	UPROPERTY()
+	int32 MaterialIndex;
+
+	/** If true, collision is enabled for this section. */
+	UPROPERTY()
+	bool bEnableCollision;
+
+	/** If true, this section will cast a shadow */
+	UPROPERTY()
+	bool bCastShadow;
+
+	/** Maximum number of triangles which have been reserved in the index buffer */
+	UPROPERTY()
+	int32 MaxTriangles;
+
+	/** All polygons in the mesh */
+	// @todo mesheditor: These are no longer serializable.  Need UProperty sparse array, or need to native serialize using normal arrays
+	TSparseArray< FEditableStaticMeshPolygon > Polygons;
+
+	/** Sparse array of triangles, that matches the triangles in the mesh index buffers.  Elements that
+	    aren't allocated will be stored as degenerates in the mesh index buffer. */
+	TSparseArray< FEditableStaticMeshTriangle > Triangles;
+
+
+	/** Converts from an index in our Triangles array to an index of a rendering triangle's first vertex in the rendering mesh's index buffer */
+	inline static FTriangleID RenderingTriangleFirstIndexToTriangleIndex( const FStaticMeshSection& RenderingSection, const uint32 RenderingTriangleFirstIndex )
+	{
+		return FTriangleID( ( RenderingTriangleFirstIndex - RenderingSection.FirstIndex ) / 3 );
+	}
+
+	/** Converts from an index of a rendering triangle's first vertex in the rendering mesh's index buffer to an index in our Triangles array */
+	inline static uint32 TriangleIndexToRenderingTriangleFirstIndex( const FStaticMeshSection& RenderingSection, const FTriangleID TriangleIndex )
+	{
+		return TriangleIndex.GetValue() * 3 + RenderingSection.FirstIndex;
+	}
+
+	/** Serializer */
+	friend FArchive& operator<<( FArchive& Ar, FEditableStaticMeshSection& Section )
+	{
+		Ar << Section.RenderingSectionIndex;
+		Ar << Section.MaterialIndex;
+		Ar << Section.bEnableCollision;
+		Ar << Section.bCastShadow;
+		Ar << Section.MaxTriangles;	// @todo mesheditor serialization: Should not need to be serialized if we triangulate after load
+		SerializeSparseArray( Ar, Section.Polygons );
+		SerializeSparseArray( Ar, Section.Triangles );	// @todo mesheditor serialization: Should not need to be serialized if we triangulate after load
+		return Ar;
+	}
+};
+
+
+UCLASS(MinimalAPI)
+class UEditableStaticMesh : public UEditableMesh
+{
+	GENERATED_BODY()
+
+public:
+
+	/** Default constructor that initializes good defaults for UEditableStaticMesh */
+	UEditableStaticMesh();
+
+	virtual void Serialize( FArchive& Ar ) override;
+
+	/** Creates a editable static mesh from the specified component and sub-mesh address */
+	void InitEditableStaticMesh( class UPrimitiveComponent& Component, const FEditableMeshSubMeshAddress& InitSubMeshAddress );
+	MESHEDITINGRUNTIME_API void InitFromBlankStaticMesh( UStaticMesh& InStaticMesh );
+
+	struct FElementIDRemappings
+	{
+		TSparseArray<FVertexID> NewVertexIndexLookup;
+		TSparseArray<FRenderingVertexID> NewRenderingVertexIndexLookup;
+		TSparseArray<FEdgeID> NewEdgeIndexLookup;
+		TSparseArray<FSectionID> NewSectionIndexLookup;
+
+		struct FPerPolygonLookups
+		{
+			TSparseArray<FPolygonID> NewPolygonIndexLookup;
+			TSparseArray<FTriangleID> NewTriangleIndexLookup;
+		};
+
+		TSparseArray<FPerPolygonLookups> PerPolygon;
+
+		FVertexID GetRemappedVertexID( FVertexID VertexID ) const
+		{
+			check( NewVertexIndexLookup.IsAllocated( VertexID.GetValue() ) );
+			return NewVertexIndexLookup[ VertexID.GetValue() ];
+		}
+
+		FRenderingVertexID GetRemappedRenderingVertexID( FRenderingVertexID RenderingVertexID ) const
+		{
+			check( NewRenderingVertexIndexLookup.IsAllocated( RenderingVertexID.GetValue() ) );
+			return NewRenderingVertexIndexLookup[ RenderingVertexID.GetValue() ];
+		}
+
+		FEdgeID GetRemappedEdgeID( FEdgeID EdgeID ) const
+		{
+			check( NewEdgeIndexLookup.IsAllocated( EdgeID.GetValue() ) );
+			return NewEdgeIndexLookup[ EdgeID.GetValue() ];
+		}
+
+		FSectionID GetRemappedSectionID( FSectionID SectionID ) const
+		{
+			check( NewSectionIndexLookup.IsAllocated( SectionID.GetValue() ) );
+			return NewSectionIndexLookup[ SectionID.GetValue() ];
+		}
+
+		FPolygonRef GetRemappedPolygonRef( FPolygonRef PolygonRef ) const
+		{
+			const FSectionID NewSectionID( GetRemappedSectionID( PolygonRef.SectionID ) );
+
+			check( PerPolygon.IsAllocated( NewSectionID.GetValue() ) );
+			const TSparseArray<FPolygonID>& NewPolygonIndexLookup = PerPolygon[ NewSectionID.GetValue() ].NewPolygonIndexLookup;
+
+			check( NewPolygonIndexLookup.IsAllocated( PolygonRef.PolygonID.GetValue() ) );
+			return FPolygonRef( NewSectionID, NewPolygonIndexLookup[ PolygonRef.PolygonID.GetValue() ] );
+		}
+
+		FTriangleID GetRemappedTriangleID( FSectionID RemappedSectionID, FTriangleID TriangleID ) const
+		{
+			check( PerPolygon.IsAllocated( RemappedSectionID.GetValue() ) );
+			const TSparseArray<FTriangleID>& NewTriangleIndexLookup = PerPolygon[ RemappedSectionID.GetValue() ].NewTriangleIndexLookup;
+
+			check( NewTriangleIndexLookup.IsAllocated( TriangleID.GetValue() ) );
+			return NewTriangleIndexLookup[ TriangleID.GetValue() ];
+		}
+
+	};
+
+	/** Compacts mesh element arrays to remove gaps, and fixes up referenced IDs */
+	void Compact();
+
+	/** Remaps mesh element arrays according to the provided remappings, in order to undo a compact operation */
+	void Uncompact( const FElementIDRemappings& Remappings );
+
+	// UEditableMesh overrides
+	virtual void RebuildRenderMesh() override;
+	virtual void StartModification( const EMeshModificationType MeshModificationType, const EMeshTopologyChange MeshTopologyChange ) override;
+	virtual void EndModification( const bool bFromUndo = false ) override;
+	virtual bool IsCommitted() const override;
+	virtual bool IsCommittedAsInstance() const override;
+	virtual void Commit() override;
+	virtual UEditableMesh* CommitInstance( UPrimitiveComponent* ComponentToInstanceTo ) override;
+	virtual void Revert() override;
+	virtual UEditableMesh* RevertInstance() override;
+	virtual void PropagateInstanceChanges() override;
+
+	virtual int32 GetRenderingVertexCount() const override;
+	virtual int32 GetRenderingVertexArraySize() const override;
+
+	virtual int32 GetVertexCount() const override;
+	virtual int32 GetVertexArraySize() const override;
+	virtual bool IsValidVertex( const FVertexID VertexID ) const override;
+	virtual FVector4 GetVertexAttribute( const FVertexID VertexID, const FName AttributeName, const int32 AttributeIndex ) const override;
+	virtual void SetVertexAttribute_Internal( const FVertexID VertexID, const FName AttributeName, const int32 AttributeIndex, const FVector4& NewAttributeValue ) override;
+	virtual int32 GetVertexConnectedEdgeCount( const FVertexID VertexID ) const override;
+	virtual FEdgeID GetVertexConnectedEdge( const FVertexID VertexID, const int32 ConnectedEdgeNumber ) const override;
+
+	virtual FVector4 GetEdgeAttribute( const FEdgeID EdgeID, const FName AttributeName, const int32 AttributeIndex ) const override;
+	virtual void SetEdgeAttribute_Internal( const FEdgeID EdgeID, const FName AttributeName, const int32 AttributeIndex, const FVector4& NewAttributeValue ) override;
+	virtual int32 GetEdgeCount() const override;
+	virtual int32 GetEdgeArraySize() const override;
+	virtual bool IsValidEdge( const FEdgeID EdgeID ) const override;
+	virtual FVertexID GetEdgeVertex( const FEdgeID EdgeID, const int32 EdgeVertexNumber ) const override;
+	virtual int32 GetEdgeConnectedPolygonCount( const FEdgeID EdgeID ) const override;
+	virtual FPolygonRef GetEdgeConnectedPolygon( const FEdgeID EdgeID, const int32 ConnectedPolygonNumber ) const override;
+
+	virtual int32 GetSectionCount() const override;
+	virtual int32 GetSectionArraySize() const override;
+	virtual bool IsValidSection( const FSectionID SectionID ) const override;
+
+	virtual int32 GetPolygonCount( const FSectionID SectionID ) const override;
+	virtual int32 GetPolygonArraySize( const FSectionID SectionID ) const override;
+	virtual bool IsValidPolygon( const FPolygonRef PolygonRef ) const override;
+	virtual int32 GetTriangleCount( const FSectionID SectionID ) const override;
+	virtual int32 GetTriangleArraySize( const FSectionID SectionID ) const override;
+	virtual int32 GetPolygonPerimeterVertexCount( const FPolygonRef PolygonRef ) const override;
+	virtual FVertexID GetPolygonPerimeterVertex( const FPolygonRef PolygonRef, const int32 PolygonVertexNumber ) const override;
+	virtual FVector4 GetPolygonPerimeterVertexAttribute( const FPolygonRef PolygonRef, const int32 PolygonVertexNumber, const FName AttributeName, const int32 AttributeIndex ) const override;
+	virtual void SetPolygonPerimeterVertexAttribute_Internal( const FPolygonRef PolygonRef, const int32 PolygonVertexNumber, const FName AttributeName, const int32 AttributeIndex, const FVector4& NewAttributeValue ) override;
+
+	virtual int32 GetPolygonHoleCount( const FPolygonRef PolygonRef ) const override;
+	virtual int32 GetPolygonHoleVertexCount( const FPolygonRef PolygonRef, const int32 HoleNumber ) const override;
+	virtual FVertexID GetPolygonHoleVertex( const FPolygonRef PolygonRef, const int32 HoleNumber, const int32 PolygonVertexNumber ) const override;
+	virtual FVector4 GetPolygonHoleVertexAttribute( const FPolygonRef PolygonRef, const int32 HoleNumber, const int32 PolygonVertexNumber, const FName AttributeName, const int32 AttributeIndex ) const override;
+	virtual void SetPolygonHoleVertexAttribute_Internal( const FPolygonRef PolygonRef, const int32 HoleNumber, const int32 PolygonVertexNumber, const FName AttributeName, const int32 AttributeIndex, const FVector4& NewAttributeValue ) override;
+	virtual int32 GetPolygonTriangulatedTriangleCount( const FPolygonRef PolygonRef ) const override;
+	virtual FVector GetPolygonTriangulatedTriangleVertexPosition( const FPolygonRef PolygonRef, const int32 PolygonTriangleNumber, const int32 TriangleVertexNumber ) const override;
+	virtual void RetriangulatePolygons( const TArray<FPolygonRef>& PolygonRefs, const bool bOnlyOnUndo ) override;
+
+protected:
+
+	virtual void CreateEmptyVertexRange_Internal( const int32 NumVerticesToAdd, const TArray<FVertexID>* OverrideVertexIDsForRedo, TArray<FVertexID>& OutNewVertexIDs ) override;
+	virtual void CreateEdge_Internal( const FVertexID VertexIDA, const FVertexID VertexIDB, const TArray<FPolygonRef>& ConnectedPolygons, const FEdgeID OverrideEdgeIDForRedo, FEdgeID& OutNewEdgeID ) override;
+	virtual void CreatePolygon_Internal( const FSectionID SectionID, const TArray<FVertexID>& VertexIDs, const TArray<TArray<FVertexID>>& VertexIDsForEachHole, const FPolygonID OverridePolygonIDForRedo, FPolygonRef& OutNewPolygonRef, TArray<FEdgeID>& OutNewEdgeIDs ) override;
+	virtual void DeleteOrphanVertices_Internal( const TArray<FVertexID>& VertexIDsToDelete ) override;
+	virtual void DeleteEdges_Internal( const TArray<FEdgeID>& EdgeIDsToDelete, const bool bDeleteOrphanedVertices ) override;
+	virtual void DeletePolygon_Internal( const FPolygonRef PolygonRef, const bool bDeleteOrphanedEdges, const bool bDeleteOrphanedVertices, const bool bDeleteEmptySections ) override;
+	virtual void SetEdgeVertices_Internal( const FEdgeID EdgeID, const FVertexID NewVertexID0, const FVertexID NewVertexID1 ) override;
+	virtual void InsertPolygonPerimeterVertices_Internal( const FPolygonRef PolygonRef, const int32 InsertBeforeVertexNumber, const TArray<FVertexAndAttributes>& VerticesToInsert ) override;
+	virtual void RemovePolygonPerimeterVertices_Internal( const FPolygonRef PolygonRef, const int32 FirstVertexNumberToRemove, const int32 NumVerticesToRemove ) override;
+	virtual FSectionID GetSectionIDFromMaterial_Internal( UMaterialInterface* Material, bool bCreateNewSectionIfNotFound ) override;
+	virtual FSectionID CreateSection_Internal( const FSectionToCreate& SectionToCreate ) override;
+	virtual void DeleteSection_Internal( const FSectionID SectionID ) override;
+
+
+private:
+
+	/** Prepares render buffers for updating */
+	void RebuildRenderMeshStart( const bool bRefreshBounds, const bool bInvalidateLighting );
+
+	/** */
+	void RebuildRenderMeshFinish( const bool bUpdateCollision );
+
+	/** */
+	void RebuildRenderMeshInternal();
+
+	/** Given a set of index remappings, fixes up references to element IDs */
+	void FixUpElementIDs( const FElementIDRemappings& Remappings );
+
+	void InitializeStaticMeshBuildVertex( FStaticMeshBuildVertex& StaticMeshVertex, const FEditableStaticMeshPolygonContourVertex ContourVertex );
+
+	/** Gets the named attribute value from a polygon contour vertex */
+	FVector4 GetPolygonContourVertexAttribute( const FEditableStaticMeshPolygonContourVertex& PolygonContourVertex, const FName AttributeName, const int32 AttributeIndex ) const;
+
+	/** Sets the named attribute value for a polygon contour vertex */
+	void SetPolygonContourVertexAttribute( FEditableStaticMeshPolygonContourVertex& PolygonContourVertex, const FName AttributeName, const int32 AttributeIndex, const FVector4& NewAttributeValue );
+
+	FRenderingVertexID GetPolygonPerimeterRenderingVertex( const FPolygonRef PolygonRef, const int32 PolygonVertexNumber ) const;
+	FRenderingVertexID GetPolygonHoleRenderingVertex( const FPolygonRef PolygonRef, const int32 HoleNumber, const int32 PolygonVertexNumber ) const;
+
+	/** Adds the specified number of new rendering vertices, and returns the IDs of the new rendering vertices */
+	void CreateRenderingVertices( const TArray<FVertexID>& VertexIDs, const TOptional<FRenderingVertexID> OptionalCopyFromRenderingVertexID, TArray<FRenderingVertexID>& OutNewRenderingVertexIDs );
+
+	/** Deletes a rendering vertex that was confirmed to be no longer used by any polygons */
+	void DeleteOrphanRenderingVertices( const TArray<FRenderingVertexID>& RenderingVertexIDsToDelete );
+
+	/** Deletes all of a polygon's triangles (including rendering triangles from the index buffer.) */
+	void DeletePolygonTriangles( const FPolygonRef PolygonRef );
+
+	/** Deletes triangles from the index buffer for the specified section and list of triangle indices */
+	void DeleteRenderingTrianglesForSectionTriangles( const FSectionID SectionID, const TArray<FTriangleID>& SectionTriangleIDsToRemove );
+
+	/** Helper function to get a specific static mesh LOD structure for a mesh */
+	const struct FStaticMeshLODResources& GetStaticMeshLOD() const;
+	struct FStaticMeshLODResources& GetStaticMeshLOD();
+
+	/** Figures out whether a rendering vertex is unique to the specified polygon */
+	bool DoesPolygonPerimeterVertexHaveDiscreetRenderingVertex( const FPolygonRef PolygonRef, const int32 PolygonVertexNumber ) const;
+
+	/** Given a polygon vertex, returns a rendering vertex that is guaranteed to be unique for that polygon */
+	FRenderingVertexID MakeDiscreetPolygonPerimeterRenderingVertexIfNeeded( const FPolygonRef PolygonRef, const int32 PolygonVertexNumber );
+
+	/** Adds a new rendering vertex to the specified polygon's perimeter at the specified vertex number of that polygon.  Returns the new rendering vertex index. */
+	FRenderingVertexID AddNewRenderingVertexToPolygonPerimeter( const FPolygonRef PolygonRef, const int32 PolygonVertexNumber, const TOptional<FRenderingVertexID> OptionalCopyFromRenderingVertexID = TOptional<FRenderingVertexID>() );
+
+	/** Makes sure our mesh's index buffer is 32-bit, converting if needed */
+	void EnsureIndexBufferIs32Bit();
+
+	/** If any of the specified vertex indices are greater than 0xffff, updates our mesh's index buffer to be 32-bit if it isn't already */
+	void UpdateIndexBufferFormatIfNeeded( const TArray<FRenderingVertexID>& RenderingVertexIDs );
+
+	/** If the specified vertex index is greater than 0xffff, updates our mesh's index buffer to be 32-bit if it isn't already */
+	void UpdateIndexBufferFormatIfNeeded( const FRenderingVertexID RenderingVertexID );
+
+	/** Rebuilds bounds */
+	void UpdateBoundsAndCollision( const bool bUpdateCollision );
+
+	/** Updates the index buffer to make room for more indices in the given section */
+	void AllocateExtraIndicesForSection( const FSectionID SectionID, int32 NumExtraTriangles = IndexBufferInterSectionGap );
+
+	/** Gets the editable mesh section index which corresponds to the given rendering section index */
+	FSectionID GetSectionForRenderingSectionIndex( const int32 RenderingSectionIndex ) const;
+
+private:
+
+	/** The static mesh asset we're representing */
+	UPROPERTY()
+	UStaticMesh* StaticMesh;
+
+	UPROPERTY()
+	UStaticMesh* OriginalStaticMesh;
+
+	/** Counter to determine when we should compact data */
+	UPROPERTY()
+	int32 PendingCompactCounter;
+
+	/** Data will be compacted after this many topology modifying actions. */
+	static const int32 CompactFrequency = 1;
+
+	/** Size of the extra gaps in the index buffer between different sections (so new triangles have some space to be added into,
+		without requiring the index buffer to be manipulated) */
+	static const uint32 IndexBufferInterSectionGap = 32;
+
+	/** Each editable vertex in this mesh. */
+	TSparseArray< FEditableStaticMeshVertex > Vertices;
+
+	/** Sparse array of rendering vertices, that matches the vertices in the mesh vertex buffers */
+	TSparseArray< FEditableStaticMeshRenderingVertex > RenderingVertices;
+
+	/** All editable mesh edges.  Note that some of these edges will be internal polygon edges, synthesized while
+	    triangulating polygons into triangles.  Static meshes currently only support triangles. */
+	TSparseArray< FEditableStaticMeshEdge > Edges;
+
+	/** All of the sections in this mesh */
+	TSparseArray< FEditableStaticMeshSection > Sections;
+
+
+	/** Used to refresh all components in the scene that may be using a mesh we're editing */
+	TSharedPtr< class FStaticMeshComponentRecreateRenderStateContext > RecreateRenderStateContext;
+};
diff --git a/Engine/Source/Runtime/MeshEditingRuntime/MeshEditingRuntime.Build.cs b/Engine/Source/Runtime/MeshEditingRuntime/MeshEditingRuntime.Build.cs
new file mode 100644
--- /dev/null
+++ b/Engine/Source/Runtime/MeshEditingRuntime/MeshEditingRuntime.Build.cs
@@ -0,0 +1,21 @@
+// Copyright 1998-2015 Epic Games, Inc. All Rights Reserved.
+
+namespace UnrealBuildTool.Rules
+{
+	public class MeshEditingRuntime : ModuleRules
+	{
+        public MeshEditingRuntime(TargetInfo Target)
+		{
+			PrivateDependencyModuleNames.AddRange(
+                new string[] {
+				    "Core",
+				    "CoreUObject",
+				    "Engine",
+					"RenderCore",	// @todo mesheditor: For FlushRenderingCommands()
+				}
+			);
+
+			AddEngineThirdPartyPrivateStaticDependencies(Target, new string[] { "MikkTSpace", "OpenSubdiv" });
+		}
+	}
+}
\ No newline at end of file
diff --git a/Engine/Source/Runtime/MeshEditingRuntime/MeshEditingRuntimeModule.cpp b/Engine/Source/Runtime/MeshEditingRuntime/MeshEditingRuntimeModule.cpp
new file mode 100644
--- /dev/null
+++ b/Engine/Source/Runtime/MeshEditingRuntime/MeshEditingRuntimeModule.cpp
@@ -0,0 +1,43 @@
+// Copyright 1998-2015 Epic Games, Inc. All Rights Reserved.
+
+#include "MeshEditingRuntimeModule.h"
+#include "IMeshEditingRuntimeModule.h"
+#include "StaticMeshEditableMeshFormat.h"
+
+DEFINE_LOG_CATEGORY( LogMeshEditingRuntime );
+
+
+class FMeshEditingRuntimeModule : public IMeshEditingRuntimeModule
+{
+public:
+
+	FMeshEditingRuntimeModule()
+	{
+	}
+
+	virtual void StartupModule() override;
+
+	virtual void ShutdownModule() override;
+
+private:
+
+	/** Static mesh editing */
+	FStaticMeshEditableMeshFormat StaticMeshEditableMeshFormat;
+
+};
+
+
+void FMeshEditingRuntimeModule::StartupModule()
+{
+	IModularFeatures::Get().RegisterModularFeature( "EditableMeshFormat", &StaticMeshEditableMeshFormat );
+}
+
+
+void FMeshEditingRuntimeModule::ShutdownModule()
+{
+	IModularFeatures::Get().UnregisterModularFeature( "EditableMeshFormat", &StaticMeshEditableMeshFormat );
+}
+
+
+
+IMPLEMENT_MODULE( FMeshEditingRuntimeModule, MeshEditingRuntime )
diff --git a/Engine/Source/Runtime/MeshEditingRuntime/MeshEditingRuntimeModule.h b/Engine/Source/Runtime/MeshEditingRuntime/MeshEditingRuntimeModule.h
new file mode 100644
--- /dev/null
+++ b/Engine/Source/Runtime/MeshEditingRuntime/MeshEditingRuntimeModule.h
@@ -0,0 +1,9 @@
+// Copyright 1998-2015 Epic Games, Inc. All Rights Reserved.
+
+#pragma once
+
+#include "CoreMinimal.h"
+#include "IMeshEditingRuntimeModule.h"
+
+
+DECLARE_LOG_CATEGORY_EXTERN( LogMeshEditingRuntime, Log, All );
diff --git a/Engine/Source/Runtime/MeshEditingRuntime/Public/EditableMesh.h b/Engine/Source/Runtime/MeshEditingRuntime/Public/EditableMesh.h
new file mode 100644
--- /dev/null
+++ b/Engine/Source/Runtime/MeshEditingRuntime/Public/EditableMesh.h
@@ -0,0 +1,388 @@
+// Copyright 1998-2015 Epic Games, Inc. All Rights Reserved.
+
+#pragma once
+
+#include "CoreMinimal.h"
+#include "UObject/Object.h"
+#include "EditableMeshTypes.h"
+#include "Change.h"		// For TUniquePtr<FChange>
+#include "EditableMesh.generated.h"
+
+
+// @todo mesheditor: Comment these classes and enums!
+
+namespace OpenSubdiv
+{
+	namespace v3_0_2
+	{
+		namespace Far
+		{
+			class TopologyRefiner;
+		}
+	}
+}
+
+
+UENUM( BlueprintType )
+enum class EInsetPolygonsMode : uint8
+{
+	All,
+	CenterPolygonOnly,
+	SidePolygonsOnly,
+};
+
+
+UENUM( BlueprintType )
+enum class ETriangleTessellationMode : uint8
+{
+	/** Connect each vertex to a new center vertex, forming three triangles */
+	ThreeTriangles,
+
+	/** Split each edge and create a center polygon that connects those new vertices, then three additional polygons for each original corner */
+	FourTriangles,
+};
+
+
+UCLASS( abstract, BlueprintType )
+class MESHEDITINGRUNTIME_API UEditableMesh : public UObject
+{
+	GENERATED_BODY()
+
+public:
+
+	/** Default constructor that initializes good defaults for UEditableMesh */
+	UEditableMesh();
+
+	// UObject interface
+	virtual void Serialize( FArchive& Ar ) override;
+	virtual void PostLoad() override;
+
+	UFUNCTION( BlueprintCallable, Category="Editable Mesh" ) virtual void RebuildRenderMesh() PURE_VIRTUAL(,);
+	UFUNCTION( BlueprintCallable, Category="Editable Mesh" ) virtual void StartModification( const EMeshModificationType MeshModificationType, const EMeshTopologyChange MeshTopologyChange ) PURE_VIRTUAL(,);
+	UFUNCTION( BlueprintCallable, Category="Editable Mesh" ) virtual void EndModification( const bool bFromUndo = false ) PURE_VIRTUAL(,);
+	UFUNCTION( BlueprintPure, Category="Editable Mesh" ) virtual bool IsCommitted() const PURE_VIRTUAL(, return false;);
+	UFUNCTION( BlueprintPure, Category="Editable Mesh" ) virtual bool IsCommittedAsInstance() const PURE_VIRTUAL(, return false;);
+	UFUNCTION( BlueprintCallable, Category="Editable Mesh" ) virtual void Commit() PURE_VIRTUAL(,);
+	UFUNCTION( BlueprintCallable, Category="Editable Mesh" ) virtual UEditableMesh* CommitInstance( UPrimitiveComponent* ComponentToInstanceTo ) PURE_VIRTUAL(,return nullptr;);
+	UFUNCTION( BlueprintCallable, Category="Editable Mesh" ) virtual void Revert() PURE_VIRTUAL(,);
+	UFUNCTION( BlueprintCallable, Category="Editable Mesh" ) virtual UEditableMesh* RevertInstance() PURE_VIRTUAL(,return nullptr;);
+	UFUNCTION( BlueprintCallable, Category="Editable Mesh" ) virtual void PropagateInstanceChanges() PURE_VIRTUAL(,);
+
+	UFUNCTION( BlueprintPure, Category="Editable Mesh" ) virtual int32 GetRenderingVertexCount() const PURE_VIRTUAL(, return 0;);
+	UFUNCTION( BlueprintPure, Category="Editable Mesh" ) virtual int32 GetRenderingVertexArraySize() const PURE_VIRTUAL(, return 0;);
+
+	UFUNCTION( BlueprintPure, Category="Editable Mesh" ) virtual int32 GetVertexCount() const PURE_VIRTUAL(,return 0;);
+	UFUNCTION( BlueprintPure, Category="Editable Mesh" ) virtual int32 GetVertexArraySize() const PURE_VIRTUAL(,return 0;);
+	UFUNCTION( BlueprintPure, Category="Editable Mesh" ) virtual bool IsValidVertex( const FVertexID VertexID ) const PURE_VIRTUAL(,return false;);
+	UFUNCTION( BlueprintPure, Category="Editable Mesh" ) virtual FVector4 GetVertexAttribute( const FVertexID VertexID, const FName AttributeName, const int32 AttributeIndex ) const PURE_VIRTUAL(,return FVector4(0.0f););
+	UFUNCTION( BlueprintPure, Category="Editable Mesh" ) virtual int32 GetVertexConnectedEdgeCount( const FVertexID VertexID ) const PURE_VIRTUAL(,return 0;);
+	UFUNCTION( BlueprintPure, Category="Editable Mesh" ) virtual FEdgeID GetVertexConnectedEdge( const FVertexID VertexID, const int32 ConnectedEdgeNumber ) const PURE_VIRTUAL(,return FEdgeID::Invalid;);
+
+	UFUNCTION( BlueprintPure, Category="Editable Mesh" ) virtual FVector4 GetEdgeAttribute( const FEdgeID EdgeID, const FName AttributeName, const int32 AttributeIndex ) const PURE_VIRTUAL(,return FVector4(0.0f););
+	UFUNCTION( BlueprintPure, Category="Editable Mesh" ) virtual int32 GetEdgeCount() const PURE_VIRTUAL(,return 0;);
+	UFUNCTION( BlueprintPure, Category="Editable Mesh" ) virtual int32 GetEdgeArraySize() const PURE_VIRTUAL(,return 0;);
+	UFUNCTION( BlueprintPure, Category="Editable Mesh" ) virtual bool IsValidEdge( const FEdgeID EdgeID ) const PURE_VIRTUAL(,return false;);
+	UFUNCTION( BlueprintPure, Category="Editable Mesh" ) virtual FVertexID GetEdgeVertex( const FEdgeID EdgeID, const int32 EdgeVertexNumber ) const PURE_VIRTUAL(,return FVertexID::Invalid;);
+	UFUNCTION( BlueprintPure, Category="Editable Mesh" ) virtual int32 GetEdgeConnectedPolygonCount( const FEdgeID EdgeID ) const PURE_VIRTUAL(,return 0;);
+	UFUNCTION( BlueprintPure, Category="Editable Mesh" ) virtual FPolygonRef GetEdgeConnectedPolygon( const FEdgeID EdgeID, const int32 ConnectedPolygonNumber ) const PURE_VIRTUAL( , return FPolygonRef(););
+
+	UFUNCTION( BlueprintPure, Category="Editable Mesh" ) virtual int32 GetSectionCount() const PURE_VIRTUAL(,return 0;);
+	UFUNCTION( BlueprintPure, Category="Editable Mesh" ) virtual int32 GetSectionArraySize() const PURE_VIRTUAL(,return 0;);
+	UFUNCTION( BlueprintPure, Category="Editable Mesh" ) virtual bool IsValidSection( const FSectionID SectionID ) const PURE_VIRTUAL(,return false;);
+
+	UFUNCTION( BlueprintPure, Category="Editable Mesh" ) virtual int32 GetPolygonCount( const FSectionID SectionID ) const PURE_VIRTUAL(,return 0;);
+	UFUNCTION( BlueprintPure, Category="Editable Mesh" ) virtual int32 GetPolygonArraySize( const FSectionID SectionID ) const PURE_VIRTUAL(,return 0;);
+	UFUNCTION( BlueprintPure, Category="Editable Mesh" ) virtual bool IsValidPolygon( const FPolygonRef PolygonRef ) const PURE_VIRTUAL(,return false;);
+	UFUNCTION( BlueprintPure, Category="Editable Mesh" ) virtual int32 GetTriangleCount( const FSectionID SectionID ) const PURE_VIRTUAL(,return 0;);
+	UFUNCTION( BlueprintPure, Category="Editable Mesh" ) virtual int32 GetTriangleArraySize( const FSectionID SectionID ) const PURE_VIRTUAL(,return 0;);
+	UFUNCTION( BlueprintPure, Category="Editable Mesh" ) virtual int32 GetPolygonPerimeterVertexCount( const FPolygonRef PolygonRef ) const PURE_VIRTUAL(,return 0;);
+	UFUNCTION( BlueprintPure, Category="Editable Mesh" ) virtual FVertexID GetPolygonPerimeterVertex( const FPolygonRef PolygonRef, const int32 PolygonVertexNumber ) const PURE_VIRTUAL(,return FVertexID::Invalid;);
+	UFUNCTION( BlueprintPure, Category="Editable Mesh" ) virtual FVector4 GetPolygonPerimeterVertexAttribute( const FPolygonRef PolygonRef, const int32 PolygonVertexNumber, const FName AttributeName, const int32 AttributeIndex ) const PURE_VIRTUAL(,return FVector4(0.0f););
+	UFUNCTION( BlueprintPure, Category="Editable Mesh" ) virtual int32 GetPolygonHoleCount( const FPolygonRef PolygonRef ) const PURE_VIRTUAL(,return 0;);
+	UFUNCTION( BlueprintPure, Category="Editable Mesh" ) virtual int32 GetPolygonHoleVertexCount( const FPolygonRef PolygonRef, const int32 HoleNumber ) const PURE_VIRTUAL(,return 0;);
+	UFUNCTION( BlueprintPure, Category="Editable Mesh" ) virtual FVertexID GetPolygonHoleVertex( const FPolygonRef PolygonRef, const int32 HoleNumber, const int32 PolygonVertexNumber ) const PURE_VIRTUAL(,return FVertexID::Invalid;);
+	UFUNCTION( BlueprintPure, Category="Editable Mesh" ) virtual FVector4 GetPolygonHoleVertexAttribute( const FPolygonRef PolygonRef, const int32 HoleNumber, const int32 PolygonVertexNumber, const FName AttributeName, const int32 AttributeIndex ) const PURE_VIRTUAL(,return FVector4(0.0f););
+	UFUNCTION( BlueprintPure, Category="Editable Mesh" ) virtual int32 GetPolygonTriangulatedTriangleCount( const FPolygonRef PolygonRef ) const PURE_VIRTUAL(,return 0;);
+	UFUNCTION( BlueprintPure, Category="Editable Mesh" ) virtual FVector GetPolygonTriangulatedTriangleVertexPosition( const FPolygonRef PolygonRef, const int32 PolygonTriangleNumber, const int32 TriangleVertexNumber ) const PURE_VIRTUAL(,return FVector(0.0f););
+
+	UFUNCTION( BlueprintCallable, Category="Editable Mesh" ) virtual void RetriangulatePolygons( const TArray<FPolygonRef>& PolygonRefs, const bool bOnlyOnUndo ) PURE_VIRTUAL(,);
+
+protected:
+
+	virtual void DeleteOrphanVertices_Internal( const TArray<FVertexID>& VertexIDsToDelete ) PURE_VIRTUAL( , );
+	virtual void SetVertexAttribute_Internal( const FVertexID VertexID, const FName AttributeName, const int32 AttributeIndex, const FVector4& NewAttributeValue ) PURE_VIRTUAL( , );
+	virtual void CreateEmptyVertexRange_Internal( const int32 NumVerticesToAdd, const TArray<FVertexID>* OverrideVertexIDsForRedo, TArray<FVertexID>& OutNewVertexIDs ) PURE_VIRTUAL(,);
+	virtual void CreateEdge_Internal( const FVertexID VertexIDA, const FVertexID VertexIDB, const TArray<FPolygonRef>& ConnectedPolygons, const FEdgeID OverrideEdgeIDForRedo, FEdgeID& OutNewEdgeID ) PURE_VIRTUAL(,);
+	virtual void SetEdgeAttribute_Internal( const FEdgeID EdgeID, const FName AttributeName, const int32 AttributeIndex, const FVector4& NewAttributeValue ) PURE_VIRTUAL( , );
+	virtual void DeleteEdges_Internal( const TArray<FEdgeID>& EdgeIDsToDelete, const bool bDeleteOrphanedVertices ) PURE_VIRTUAL(,);
+	virtual void SetEdgeVertices_Internal( const FEdgeID EdgeID, const FVertexID NewVertexID0, const FVertexID NewVertexID1 ) PURE_VIRTUAL(,);
+	virtual void CreatePolygon_Internal( const FSectionID SectionID, const TArray<FVertexID>& VertexIDs, const TArray<TArray<FVertexID>>& VertexIDsForEachHole, const FPolygonID OverridePolygonIDForRedo, FPolygonRef& OutNewPolygonRef, TArray<FEdgeID>& OutNewEdgeIDs ) PURE_VIRTUAL( ,);
+	virtual void DeletePolygon_Internal( const FPolygonRef PolygonRef, const bool bDeleteOrphanedEdges, const bool bDeleteOrphanedVertices, const bool bDeleteEmptySections ) PURE_VIRTUAL(,);
+	virtual void SetPolygonPerimeterVertexAttribute_Internal( const FPolygonRef PolygonRef, const int32 PolygonVertexNumber, const FName AttributeName, const int32 AttributeIndex, const FVector4& NewAttributeValue ) PURE_VIRTUAL( , );
+	virtual void SetPolygonHoleVertexAttribute_Internal( const FPolygonRef PolygonRef, const int32 HoleNumber, const int32 PolygonVertexNumber, const FName AttributeName, const int32 AttributeIndex, const FVector4& NewAttributeValue ) PURE_VIRTUAL( , );
+	virtual void InsertPolygonPerimeterVertices_Internal( const FPolygonRef PolygonRef, const int32 InsertBeforeVertexNumber, const TArray<FVertexAndAttributes>& VerticesToInsert ) PURE_VIRTUAL(,);
+	virtual void RemovePolygonPerimeterVertices_Internal( const FPolygonRef PolygonRef, const int32 FirstVertexNumberToRemove, const int32 NumVerticesToRemove ) PURE_VIRTUAL(,);
+	virtual FSectionID GetSectionIDFromMaterial_Internal( UMaterialInterface* Material, bool bCreateNewSectionIfNotFound ) PURE_VIRTUAL(, return FSectionID::Invalid; );
+	virtual FSectionID CreateSection_Internal( const FSectionToCreate& SectionToCreate ) PURE_VIRTUAL(,return FSectionID::Invalid;);
+	virtual void DeleteSection_Internal( const FSectionID SectionID ) PURE_VIRTUAL(,);
+
+public:
+
+	/**
+	 * Called at initialization time to set this mesh's sub-mesh address
+	 *
+	 * @param	NewSubMeshAddress	The new sub-mesh address for this mesh
+	 */
+	void SetSubMeshAddress( const FEditableMeshSubMeshAddress& NewSubMeshAddress );
+
+	/**
+	 * @return	Returns true if StartModification() was called and the mesh is able to be modified currently.  Remember to call EndModification() when finished.
+	 */
+	UFUNCTION( BlueprintPure, Category="Editable Mesh" ) virtual bool IsBeingModified() const
+	{
+		return bIsBeingModified;
+	}
+
+	/**
+	 * @return	Returns true if undo tracking is enabled on this mesh
+	 */
+	UFUNCTION( BlueprintPure, Category="Editable Mesh" )
+	bool IsUndoAllowed() const
+	{
+		return bAllowUndo;
+	}
+
+	/**
+	 * Sets whether undo is allowed on this mesh
+	 *
+	 * @param	bInAllowUndo	True if undo features are enabled on this mesh.  You're only allowed to call MakeUndo() if this is set to true.
+	 */
+	UFUNCTION( BlueprintCallable, Category="Editable Mesh" )
+	void SetAllowUndo( const bool bInAllowUndo )
+	{
+		bAllowUndo = bInAllowUndo;
+	}
+
+
+	/**
+	 * @return	Returns true if there are any existing tracked changes that can be undo right now
+	 */
+	UFUNCTION( BlueprintPure, Category="Editable Mesh" )
+	bool AnyChangesToUndo() const;
+
+
+	/**
+	 * Gets the sub-mesh address for this mesh which uniquely identifies the mesh among other sub-meshes in the same component
+	 *
+	 * @return	The sub-mesh address for the mesh
+	 */
+	const FEditableMeshSubMeshAddress& GetSubMeshAddress() const;
+
+	/** Grabs any outstanding changes to this mesh and returns a change that can be used to undo those changes.  Calling this
+	    function will clear the history of changes.  This function will return a null pointer if bAllowUndo is false. */
+	// @todo mesheditor script: We might need this to be available for BP editable meshes, in some form at least.  Probably it should just apply the undo right away.
+	TUniquePtr<FChange> MakeUndo();
+
+
+	/**
+	 * Statics
+	 */
+
+	UFUNCTION( BlueprintPure, Category="Editable Mesh" ) static const TArray<FName>& GetValidVertexAttributes();
+	UFUNCTION( BlueprintPure, Category="Editable Mesh" ) static const TArray<FName>& GetValidPolygonVertexAttributes();
+	UFUNCTION( BlueprintPure, Category="Editable Mesh" ) static const TArray<FName>& GetValidEdgeAttributes();
+	UFUNCTION( BlueprintPure, Category="Editable Mesh" ) static inline FVertexID InvalidVertexID()
+	{
+		return FVertexID::Invalid;
+	}
+	UFUNCTION( BlueprintPure, Category="Editable Mesh" ) static inline FEdgeID InvalidEdgeID()
+	{
+		return FEdgeID::Invalid;
+	}
+	UFUNCTION( BlueprintPure, Category="Editable Mesh" ) static inline FSectionID InvalidSectionID()
+	{
+		return FSectionID::Invalid;
+	}
+	UFUNCTION( BlueprintPure, Category="Editable Mesh" ) static inline FPolygonID InvalidPolygonID()
+	{
+		return FPolygonID::Invalid;
+	}
+	UFUNCTION( BlueprintPure, Category="Editable Mesh" ) static inline FPolygonRef InvalidPolygonRef()
+	{
+		return FPolygonRef::Invalid;
+	}
+	UFUNCTION( BlueprintPure, Category="Editable Mesh" ) static inline FVertexID MakeVertexID( const int32 VertexIndex )
+	{
+		return FVertexID( VertexIndex );
+	}
+	UFUNCTION( BlueprintPure, Category="Editable Mesh" ) static inline FEdgeID MakeEdgeID( const int32 EdgeIndex )
+	{
+		return FEdgeID( EdgeIndex );
+	}
+	UFUNCTION( BlueprintPure, Category="Editable Mesh" ) static inline FSectionID MakeSectionID( const int32 SectionIndex )
+	{
+		return FSectionID( SectionIndex );
+	}
+	UFUNCTION( BlueprintPure, Category="Editable Mesh" ) static inline FPolygonID MakePolygonID( const int32 PolygonIndex )
+	{
+		return FPolygonID( PolygonIndex );
+	}
+
+	UFUNCTION( BlueprintPure, Category="Editable Mesh" ) int32 GetMaxAttributeIndex( const FName AttributeName ) const;
+	UFUNCTION( BlueprintPure, Category="Editable Mesh" ) FSectionID GetFirstValidSection() const;
+	UFUNCTION( BlueprintPure, Category="Editable Mesh" ) int32 GetTotalPolygonCount() const;
+	UFUNCTION( BlueprintPure, Category="Editable Mesh" ) int32 GetTextureCoordinateCount() const;
+	UFUNCTION( BlueprintPure, Category="Editable Mesh" ) int32 GetSubdivisionCount() const;
+	UFUNCTION( BlueprintPure, Category="Editable Mesh" ) bool IsPreviewingSubdivisions() const;
+	UFUNCTION( BlueprintPure, Category="Editable Mesh" ) void GetVertexConnectedPolygons( const FVertexID VertexID, TArray<FPolygonRef>& OutConnectedPolygonRefs ) const;
+	UFUNCTION( BlueprintPure, Category="Editable Mesh" ) void GetVertexAdjacentVertices( const FVertexID VertexID, TArray< FVertexID >& OutAdjacentVertexIDs ) const;
+	UFUNCTION( BlueprintPure, Category="Editable Mesh" ) void GetEdgeVertices( const FEdgeID EdgeID, FVertexID& OutEdgeVertexID0, FVertexID& OutEdgeVertexID1 ) const;
+	UFUNCTION( BlueprintPure, Category="Editable Mesh" ) void GetEdgeConnectedPolygons( const FEdgeID EdgeID, TArray<FPolygonRef>& OutConnectedPolygonRefs ) const;
+	UFUNCTION( BlueprintPure, Category="Editable Mesh" ) void GetEdgeLoopElements( const FEdgeID EdgeID, TArray<FEdgeID>& EdgeLoopIDs ) const;
+	UFUNCTION( BlueprintPure, Category="Editable Mesh" ) int32 GetPolygonPerimeterEdgeCount( const FPolygonRef PolygonRef ) const;
+	UFUNCTION( BlueprintPure, Category="Editable Mesh" ) int32 GetPolygonHoleEdgeCount( const FPolygonRef PolygonRef, const int32 HoleNumber ) const;
+	UFUNCTION( BlueprintPure, Category="Editable Mesh" ) void GetPolygonPerimeterVertices( const FPolygonRef PolygonRef, TArray<FVertexID>& OutPolygonPerimeterVertexIDs ) const;
+	UFUNCTION( BlueprintPure, Category="Editable Mesh" ) void GetPolygonHoleVertices( const FPolygonRef PolygonRef, const int32 HoleNumber, TArray<FVertexID>& OutHoleVertexIDs ) const;
+	UFUNCTION( BlueprintPure, Category="Editable Mesh" ) FEdgeID GetPolygonPerimeterEdge( const FPolygonRef PolygonRef, const int32 PerimeterEdgeNumber, bool& bOutEdgeWindingIsReversedForPolygon ) const;
+	UFUNCTION( BlueprintPure, Category="Editable Mesh" ) FEdgeID GetPolygonHoleEdge( const FPolygonRef PolygonRef, const int32 HoleNumber, const int32 HoleEdgeNumber ) const;
+	UFUNCTION( BlueprintPure, Category="Editable Mesh" ) void GetPolygonPerimeterEdges( const FPolygonRef PolygonRef, TArray<FEdgeID>& OutPolygonPerimeterEdgeIDs ) const;
+	UFUNCTION( BlueprintPure, Category="Editable Mesh" ) void GetPolygonHoleEdges( const FPolygonRef PolygonRef, const int32 HoleNumber, TArray<FEdgeID>& OutHoleEdgeIDs ) const;
+	UFUNCTION( BlueprintPure, Category="Editable Mesh" ) void GetPolygonAdjacentPolygons( const FPolygonRef PolygonRef, TArray<FPolygonRef>& OutAdjacentPolygons ) const;
+	UFUNCTION( BlueprintPure, Category="Editable Mesh" ) int32 FindPolygonPerimeterVertexNumberForVertex( const FPolygonRef PolygonRef, const FVertexID VertexID ) const;
+	UFUNCTION( BlueprintPure, Category="Editable Mesh" ) int32 FindPolygonHoleVertexNumberForVertex( const FPolygonRef PolygonRef, const int32 HoleNumber, const FVertexID VertexID ) const;
+	UFUNCTION( BlueprintPure, Category="Editable Mesh" ) int32 FindPolygonPerimeterEdgeNumberForVertices( const FPolygonRef PolygonRef, const FVertexID EdgeVertexID0, const FVertexID EdgeVertexID1 ) const;
+	UFUNCTION( BlueprintPure, Category="Editable Mesh" ) int32 FindPolygonHoleEdgeNumberForVertices( const FPolygonRef PolygonRef, const int32 HoleNumber, const FVertexID EdgeVertexID0, const FVertexID EdgeVertexID1 ) const;
+	UFUNCTION( BlueprintPure, Category="Editable Mesh" ) FBox ComputeBoundingBox() const;
+	UFUNCTION( BlueprintPure, Category="Editable Mesh" ) FBoxSphereBounds ComputeBoundingBoxAndSphere() const;
+	UFUNCTION( BlueprintPure, Category="Editable Mesh" ) FVector ComputePolygonCenter( const FPolygonRef PolygonRef ) const;
+	UFUNCTION( BlueprintPure, Category="Editable Mesh" ) FPlane ComputePolygonPlane( const FPolygonRef PolygonRef ) const;
+	UFUNCTION( BlueprintPure, Category="Editable Mesh" ) FVector ComputePolygonNormal( const FPolygonRef PolygonRef ) const;
+	UFUNCTION( BlueprintPure, Category="Editable Mesh" ) FVector ComputePolygonPerimeterVertexNormal( const FPolygonRef PolygonRef, const int32 PolygonVertexNumber ) const;
+	void ComputePolygonTriangulation( const FPolygonRef PolygonRef, TArray<int32>& OutPerimeterVertexNumbersForTriangles ) const;
+	UFUNCTION( BlueprintPure, Category="Editable Mesh" ) const FSubdivisionLimitData& GetSubdivisionLimitData() const;
+	UFUNCTION( BlueprintPure, Category="Editable Mesh" ) void ComputePolygonTriangulation( const FPolygonRef PolygonRef, TArray<FVertexID>& OutTrianglesVertexIDs ) const;
+	UFUNCTION( BlueprintPure, Category="Editable Mesh" ) bool ComputeBarycentricWeightForPointOnPolygon( const FPolygonRef PolygonRef, const FVector PointOnPolygon, TArray<int32>& PerimeterVertexIndices, FVector& TriangleVertexWeights ) const;
+	UFUNCTION( BlueprintPure, Category="Editable Mesh" ) void ComputeTextureCoordinatesForPointOnPolygon( const FPolygonRef PolygonRef, const FVector PointOnPolygon, bool& bOutWereTextureCoordinatesFound, TArray<FVector4>& OutInterpolatedTextureCoordinates );
+	UFUNCTION( BlueprintPure, Category="Editable Mesh" ) void ComputePolygonsSharedEdges( const TArray<FPolygonRef>& PolygonRefs, TArray<FEdgeID>& OutSharedEdgeIDs ) const;
+	UFUNCTION( BlueprintPure, Category="Editable Mesh" ) void FindPolygonLoop( const FEdgeID EdgeID, TArray<FEdgeID>& OutEdgeLoopEdgeIDs, TArray<FEdgeID>& OutFlippedEdgeIDs, TArray<FEdgeID>& OutReversedEdgeIDPathToTake, TArray<FPolygonRef>& OutPolygonRefsToSplit ) const;
+
+
+	// Low poly editing features
+	// @todo mesheditor: Move elsewhere
+	UFUNCTION( BlueprintCallable, Category="Editable Mesh" ) void SetSubdivisionCount( const int32 NewSubdivisionCount );
+	UFUNCTION( BlueprintCallable, Category="Editable Mesh" ) void MoveVertices( const TArray<FVertexToMove>& VerticesToMove );
+	UFUNCTION( BlueprintCallable, Category="Editable Mesh" ) void CreateMissingPolygonPerimeterEdges( const FPolygonRef PolygonRef, TArray<FEdgeID>& OutNewEdgeIDs );
+	UFUNCTION( BlueprintCallable, Category="Editable Mesh" ) void CreateMissingPolygonHoleEdges( const FPolygonRef PolygonRef, const int32 HoleNumber, TArray<FEdgeID>& OutNewEdgeIDs );
+	UFUNCTION( BlueprintCallable, Category="Editable Mesh" ) void SplitEdge( const FEdgeID EdgeID, const TArray<float>& Splits, TArray<FVertexID>& OutNewVertexIDs );
+	UFUNCTION( BlueprintCallable, Category="Editable Mesh" ) void InsertEdgeLoop( const FEdgeID EdgeID, const TArray<float>& Splits, TArray<FEdgeID>& OutNewEdgeIDs );
+	UFUNCTION( BlueprintCallable, Category="Editable Mesh" ) void SplitPolygons( const TArray<FPolygonToSplit>& PolygonsToSplit, TArray<FEdgeID>& OutNewEdgeIDs );
+	UFUNCTION( BlueprintCallable, Category="Editable Mesh" ) void DeleteEdgeAndConnectedPolygons( const FEdgeID EdgeID, const bool bDeleteOrphanedEdges, const bool bDeleteOrphanedVertices, const bool bDeleteEmptySections );
+	UFUNCTION( BlueprintCallable, Category="Editable Mesh" ) void DeleteVertexAndConnectedEdgesAndPolygons( const FVertexID VertexID, const bool bDeleteOrphanedEdges, const bool bDeleteOrphanedVertices, const bool bDeleteEmptySections );
+	UFUNCTION( BlueprintCallable, Category="Editable Mesh" ) void CreateEmptyVertexRange( const int32 NumVerticesToAdd, TArray<FVertexID>& OutNewVertexIDs );
+	UFUNCTION( BlueprintCallable, Category="Editable Mesh" ) void DeleteOrphanVertices( const TArray<FVertexID>& VertexIDsToDelete );
+	UFUNCTION( BlueprintCallable, Category="Editable Mesh" ) void DeleteEdges( const TArray<FEdgeID>& EdgeIDsToDelete, const bool bDeleteOrphanedVertices );
+	UFUNCTION( BlueprintCallable, Category="Editable Mesh" ) void CreateVertices( const TArray<FVertexToCreate>& VerticesToCreate, TArray<FVertexID>& OutNewVertexIDs );
+	UFUNCTION( BlueprintCallable, Category="Editable Mesh" ) void CreateEdges( const TArray<FEdgeToCreate>& EdgesToCreate, TArray<FEdgeID>& OutNewEdgeIDs );
+	UFUNCTION( BlueprintCallable, Category="Editable Mesh" ) void CreatePolygons( const TArray<FPolygonToCreate>& PolygonsToCreate, TArray<FPolygonRef>& OutNewPolygonRefs, TArray<FEdgeID>& OutNewEdgeIDs );
+	UFUNCTION( BlueprintCallable, Category="Editable Mesh" ) void DeletePolygons( const TArray<FPolygonRef>& PolygonRefsToDelete, const bool bDeleteOrphanedEdges, const bool bDeleteOrphanedVertices, const bool bDeleteEmptySections );
+	UFUNCTION( BlueprintCallable, Category="Editable Mesh" ) void SetVerticesAttributes( const TArray<FAttributesForVertex>& AttributesForVertices );
+	UFUNCTION( BlueprintCallable, Category="Editable Mesh" ) void SetEdgesAttributes( const TArray<FAttributesForEdge>& AttributesForEdges );
+	UFUNCTION( BlueprintCallable, Category="Editable Mesh" ) void SetPolygonsVertexAttributes( const TArray<FVertexAttributesForPolygon>& VertexAttributesForPolygons );
+	UFUNCTION( BlueprintCallable, Category="Editable Mesh" ) void TryToRemovePolygonEdge( const FEdgeID EdgeID, bool& bOutWasEdgeRemoved, FPolygonRef& OutNewPolygonRef );
+	UFUNCTION( BlueprintCallable, Category="Editable Mesh" ) void TryToRemoveVertex( const FVertexID VertexID, bool& bOutWasVertexRemoved, FEdgeID& OutNewEdgeID );
+	UFUNCTION( BlueprintCallable, Category="Editable Mesh" ) void ExtrudePolygons( const TArray<FPolygonRef>& Polygons, const float ExtrudeDistance, TArray<FPolygonRef>& OutNewExtrudedFrontPolygons );
+	UFUNCTION( BlueprintCallable, Category="Editable Mesh" ) void ExtendEdges( const TArray<FEdgeID>& EdgeIDs, const bool bWeldNeighbors, TArray<FEdgeID>& OutNewExtendedEdgeIDs );
+	UFUNCTION( BlueprintCallable, Category="Editable Mesh" ) void ExtendVertices( const TArray<FVertexID>& VertexIDs, const bool bOnlyExtendClosestEdge, const FVector ReferencePosition, TArray<FVertexID>& OutNewExtendedVertexIDs );
+	UFUNCTION( BlueprintCallable, Category="Editable Mesh" ) void InsetPolygons( const TArray<FPolygonRef>& PolygonRefs, const float InsetFixedDistance, const float InsetProgressTowardCenter, const EInsetPolygonsMode Mode, TArray<FPolygonRef>& OutNewCenterPolygonRefs, TArray<FPolygonRef>& OutNewSidePolygonRefs );
+	UFUNCTION( BlueprintCallable, Category="Editable Mesh" ) void BevelPolygons( const TArray<FPolygonRef>& PolygonRefs, const float BevelFixedDistance, const float BevelProgressTowardCenter, TArray<FPolygonRef>& OutNewCenterPolygonRefs, TArray<FPolygonRef>& OutNewSidePolygonRefs );
+	UFUNCTION( BlueprintCallable, Category="Editable Mesh" ) void GenerateNormalsAndTangentsForPolygons( const TArray< FPolygonRef >& PolygonRefs );
+	UFUNCTION( BlueprintCallable, Category="Editable Mesh" ) void GenerateNormalsAndTangentsForPolygonsAndAdjacents( const TArray< FPolygonRef >& PolygonRefs );
+	UFUNCTION( BlueprintCallable, Category="Editable Mesh" ) void SetVerticesCornerSharpness( const TArray<FVertexID>& VertexIDs, const TArray<float>& VerticesNewCornerSharpness );
+	UFUNCTION( BlueprintCallable, Category="Editable Mesh" ) void SetEdgesCreaseSharpness( const TArray<FEdgeID>& EdgeIDs, const TArray<float>& EdgesNewCreaseSharpness );
+	UFUNCTION( BlueprintCallable, Category="Editable Mesh" ) void SetEdgesHardness( const TArray<FEdgeID>& EdgeIDs, const TArray<bool>& EdgesNewIsHard );
+	UFUNCTION( BlueprintCallable, Category="Editable Mesh" ) void SetEdgesHardnessAutomatically( const TArray<FEdgeID>& EdgeIDs, const float MaxDotProductForSoftEdge );	// @todo mesheditor: Not used for anything yet.  Remove it?  Use it during import/convert?
+	UFUNCTION( BlueprintCallable, Category="Editable Mesh" ) void SetEdgesVertices( const TArray<FVerticesForEdge>& VerticesForEdges );
+	UFUNCTION( BlueprintCallable, Category="Editable Mesh" ) void InsertPolygonPerimeterVertices( const FPolygonRef PolygonRef, const int32 InsertBeforeVertexNumber, const TArray<FVertexAndAttributes>& VerticesToInsert );
+	UFUNCTION( BlueprintCallable, Category="Editable Mesh" ) void RemovePolygonPerimeterVertices( const FPolygonRef PolygonRef, const int32 FirstVertexNumberToRemove, const int32 NumVerticesToRemove );
+	UFUNCTION( BlueprintCallable, Category="Editable Mesh" ) void FlipPolygons( const TArray<FPolygonRef>& PolygonRefs );
+	UFUNCTION( BlueprintCallable, Category="Editable Mesh" ) void TriangulatePolygons( const TArray<FPolygonRef>& PolygonRefs, TArray<FPolygonRef>& OutNewTrianglePolygons );
+	UFUNCTION( BlueprintCallable, Category="Editable Mesh" ) FSectionID CreateSection( const FSectionToCreate& SectionToCreate );
+	UFUNCTION( BlueprintCallable, Category="Editable Mesh" ) void DeleteSection( const FSectionID SectionID );
+	UFUNCTION( BlueprintCallable, Category="Editable Mesh" ) void AssignMaterialToPolygons( const TArray<FPolygonRef>& PolygonRefs, UMaterialInterface* Material, TArray<FPolygonRef>& NewPolygonRefs );
+	UFUNCTION( BlueprintCallable, Category="Editable Mesh" ) void WeldVertices( const TArray<FVertexID>& VertexIDs, FVertexID& OutNewVertexID );
+	UFUNCTION( BlueprintCallable, Category="Editable Mesh" ) void TessellatePolygons( const TArray<FPolygonRef>& PolygonRefs, const ETriangleTessellationMode TriangleTessellationMode, TArray<FPolygonRef>& OutNewPolygonRefs );
+	UFUNCTION( BlueprintCallable, Category="Editable Mesh" ) void SetTextureCoordinateCount(int32 NumTexCoords);
+	UFUNCTION( BlueprintCallable, Category="Editable Mesh" ) void QuadrangulateMesh( TArray<FPolygonRef>& NewPolygonRefs );
+protected:
+
+	/** Adds a new change that can be used to undo a modification that happened to the mesh.  If bAllowUndo is false, then
+	    the undo will not be stored */
+	void AddUndo( TUniquePtr<FChange> NewUndo );
+
+	/** Refreshes the entire OpenSubdiv state for this mesh and generates subdivision geometry (if the mesh is configured to have subdivision levels) */
+	void RefreshOpenSubdiv();
+
+	/** Generates new limit surface data and caches it (if the mesh is configured to have subdivision levels).  This can be used to update the
+	    limit surface after a non-topology-effective edit has happened */
+	void GenerateOpenSubdivLimitSurfaceData();
+
+
+private:
+
+	void BevelOrInsetPolygons( const TArray<FPolygonRef>& PolygonRefs, const float InsetFixedDistance, const float InsetProgressTowardCenter, const EInsetPolygonsMode Mode, const bool bShouldBevel, TArray<FPolygonRef>& OutNewCenterPolygonRefs, TArray<FPolygonRef>& OutNewSidePolygonRefs );
+
+	void GenerateTangentsForPolygons( const TArray< FPolygonRef >& PolygonRefs );
+
+
+protected:
+
+	/** The sub-mesh we came from */
+	FEditableMeshSubMeshAddress SubMeshAddress;
+
+	/** True if undo features are enabled on this mesh.  You're only allowed to call MakeUndo() if this is set to true. */
+	bool bAllowUndo;
+
+	/** When bAllowUndo is enabled, this will store the changes that can be applied to revert anything that happened to this
+	    mesh since the last time that MakeUndo() was called. */
+	TUniquePtr<FCompoundChangeInput> Undo;
+
+	/** The number of texture coordinates stored on the vertices of this mesh */
+	UPROPERTY( BlueprintReadOnly, Category="Editable Mesh" )
+	int32 TextureCoordinateCount;
+
+	/** How many levels to subdivide this mesh.  Zero will turn off subdivisions */
+	UPROPERTY( BlueprintReadOnly, Category="Editable Mesh" ) 
+	int32 SubdivisionCount;
+
+	/** True if StartModification() has been called.  Call EndModification() when you've finished changing the mesh. */
+	bool bIsBeingModified;
+
+	/** While the mesh is being edited (between calls to StartModification() and EndModification()), this is the type of modification being performed */
+	EMeshModificationType CurrentModificationType;
+
+	/** While the mesh is being edited (between calls to StartModification() and EndModification()), stores whether topology could be affected */
+	EMeshTopologyChange CurrentToplogyChange;
+
+	/** OpenSubdiv topology refiner object.  This is generated for meshes that have subdivision levels, and reused to generate new limit surfaces 
+	    when geometry is moved.  When the mesh's topology changes, this object is regenerated from scratch. */
+	TSharedPtr<OpenSubdiv::v3_0_2::Far::TopologyRefiner> OsdTopologyRefiner;
+
+	/** Various cached arrays of mesh data in the form that OpenSubdiv expects to read it.  Required by GenerateOpenSubdivLimitSurfaceData(). */
+	TArray<int> OsdNumVerticesPerFace;
+	TArray<int> OsdVertexIndicesPerFace;
+	TArray<int> OsdCreaseVertexIndexPairs;
+	TArray<float> OsdCreaseWeights;
+	TArray<int> OsdCornerVertexIndices;
+	TArray<float> OsdCornerWeights;
+
+	TArray<int> OsdFVarIndicesPerFace;
+	struct FOsdFVarChannel
+	{
+		int ValueCount;
+		int const* ValueIndices;
+	};
+	TArray<FOsdFVarChannel> OsdFVarChannels;
+
+	/** The resulting limit surface geometry after GenerateOpenSubdivLimitSurfaceData() is called */
+	FSubdivisionLimitData SubdivisionLimitData;
+};
diff --git a/Engine/Source/Runtime/MeshEditingRuntime/Public/EditableMeshFactory.h b/Engine/Source/Runtime/MeshEditingRuntime/Public/EditableMeshFactory.h
new file mode 100644
--- /dev/null
+++ b/Engine/Source/Runtime/MeshEditingRuntime/Public/EditableMeshFactory.h
@@ -0,0 +1,28 @@
+// Copyright 1998-2015 Epic Games, Inc. All Rights Reserved.
+
+#pragma once
+
+#include "CoreMinimal.h"
+#include "EditableMeshTypes.h"
+#include "EditableMeshFactory.generated.h"
+
+
+// @todo mesheditor: Comment these classes and enums!
+
+
+UCLASS()
+class MESHEDITINGRUNTIME_API UEditableMeshFactory : public UObject
+{
+	GENERATED_BODY()
+
+public:
+
+	static FEditableMeshSubMeshAddress MakeSubmeshAddress( class UPrimitiveComponent* PrimitiveComponent, const int32 LODIndex );
+	static FEditableMeshSubMeshAddress MakeSubmeshAddress( class UStaticMesh& StaticMesh, const int32 LODIndex );
+
+	static UEditableMesh* MakeEditableMesh( class UPrimitiveComponent* PrimitiveComponent, const FEditableMeshSubMeshAddress& SubMeshAddress );
+
+	UFUNCTION( BlueprintCallable, Category = "Editable Mesh" )
+	static UEditableMesh* MakeEditableMesh( class UPrimitiveComponent* PrimitiveComponent, const int32 LODIndex );
+
+};
diff --git a/Engine/Source/Runtime/MeshEditingRuntime/Public/EditableMeshTypes.h b/Engine/Source/Runtime/MeshEditingRuntime/Public/EditableMeshTypes.h
new file mode 100644
--- /dev/null
+++ b/Engine/Source/Runtime/MeshEditingRuntime/Public/EditableMeshTypes.h
@@ -0,0 +1,1170 @@
+// Copyright 1998-2015 Epic Games, Inc. All Rights Reserved.
+
+#pragma once
+
+#include "CoreMinimal.h"
+#include "UObject/Object.h"
+#include "UObject/ScriptMacros.h"
+#include "Materials/MaterialInterface.h"
+#include "EditableMeshTypes.generated.h"
+
+// @todo mesheditor: Move elsewhere
+namespace LogHelpers
+{
+	template<typename ArrayType>
+	inline FString ArrayToString( const TArray<ArrayType>& Array )
+	{
+		FString String;
+		if( Array.Num() == 0 )
+		{
+			String = TEXT( "Empty" );
+		}
+		else
+		{
+			for( const ArrayType& Element : Array )
+			{
+				if( !String.IsEmpty() )
+				{
+					String.Append( TEXT( ", " ) );
+				}
+				String.Append( Element.ToString() );
+			}
+			String = TEXT( "[" ) + String + TEXT( "]" );
+		}
+		return String;
+	}
+
+	template<>
+	inline FString ArrayToString<int32>( const TArray<int32>& Array )
+	{
+		FString String;
+		if( Array.Num() == 0 )
+		{
+			String = TEXT( "Empty" );
+		}
+		else
+		{
+			for( const int32& Element : Array )
+			{
+				if( !String.IsEmpty() )
+				{
+					String.Append( TEXT( ", " ) );
+				}
+				String.Append( FString::FormatAsNumber( Element ) );
+			}
+			String = TEXT( "[" ) + String + TEXT( "]" );
+		}
+		return String;
+	}
+
+
+	template<typename ArrayType>
+	inline FString ArrayToString( const TArray<TArray<ArrayType>>& Array )
+	{
+		FString String;
+		if( Array.Num() == 0 )
+		{
+			String = TEXT( "Empty" );
+		}
+		else
+		{
+			for( const TArray<ArrayType>& SubArray : Array )
+			{
+				if( !String.IsEmpty() )
+				{
+					String.Append( TEXT( ", " ) );
+				}
+				String.Append( FString::Printf( TEXT( "%s" ), *ArrayToString( SubArray ) ) );
+			}
+			String = TEXT( "[" ) + String + TEXT( "]" );
+		}
+		return String;
+	}
+
+	inline FString BoolToString( const bool Bool )
+	{
+		return Bool ? TEXT( "true" ) : TEXT( "false" );
+	}
+}
+
+
+/**
+ * The different components that make up a typical mesh
+ */
+enum class EEditableMeshElementType
+{
+	/** Invalid mesh element (or "none") */
+	Invalid,
+
+	/** A unique point in 3D space */
+	Vertex,
+
+	/** An edge that connects two vertices */
+	Edge,
+
+	/** A polygon with at least three 3D points.  It could be triangle, quad, or more complex shape */
+	Polygon,
+
+	/** Represents any element type */
+	Any,
+};
+
+
+// @todo mesheditor: Need comments
+
+USTRUCT( BlueprintType )
+struct FElementID	// @todo mesheditor script: BP doesn't have name spaces, so we might need a more specific display name, or just rename our various types
+{
+	GENERATED_BODY()
+
+	FElementID()
+	{
+	}
+
+	explicit FElementID( const int32 InitIDValue )
+		: IDValue( InitIDValue )
+	{
+	}
+
+	FORCEINLINE int32 GetValue() const
+	{
+		return IDValue;
+	}
+
+	FORCEINLINE bool operator==( const FElementID& Other ) const
+	{
+		return IDValue == Other.IDValue;
+	}
+
+	FORCEINLINE bool operator!=( const FElementID& Other ) const
+	{
+		return IDValue != Other.IDValue;
+	}
+
+	FString ToString() const
+	{
+		return FString::Printf( TEXT( "%lu" ), GetValue() );
+	}
+
+	/** Hashing */
+	FORCEINLINE friend uint32 GetTypeHash( const FElementID& Other )
+	{
+		return GetTypeHash( Other.IDValue );
+	}
+
+	friend FArchive& operator<<( FArchive& Ar, FElementID& Element )
+	{
+		Ar << Element.IDValue;
+		return Ar;
+	}
+
+	/** Invalid element ID */
+	MESHEDITINGRUNTIME_API static const FElementID Invalid;
+
+protected:
+
+	/** The actual mesh element index this ID represents.  Read-only. */
+	UPROPERTY( BlueprintReadOnly, Category="Editable Mesh" )
+	int32 IDValue;
+};
+
+
+USTRUCT( BlueprintType )
+struct FVertexID : public FElementID
+{
+	GENERATED_BODY()
+
+	FVertexID()
+	{
+	}
+
+	explicit FVertexID( const FElementID InitElementID )
+		: FElementID( InitElementID.GetValue() )
+	{
+	}
+
+	explicit FVertexID( const int32 InitIDValue )
+		: FElementID( InitIDValue )
+	{
+	}
+
+	/** Invalid vertex ID */
+	MESHEDITINGRUNTIME_API static const FVertexID Invalid;
+};
+
+
+USTRUCT( BlueprintType )
+struct FEdgeID : public FElementID
+{
+	GENERATED_BODY()
+
+	FEdgeID()
+	{
+	}
+
+	explicit FEdgeID( const FElementID InitElementID )
+		: FElementID( InitElementID.GetValue() )
+	{
+	}
+
+	explicit FEdgeID( const int32 InitIDValue )
+		: FElementID( InitIDValue )
+	{
+	}
+
+	/** Invalid edge ID */
+	MESHEDITINGRUNTIME_API static const FEdgeID Invalid;
+};
+
+
+USTRUCT( BlueprintType )
+struct FSectionID : public FElementID
+{
+	GENERATED_BODY()
+
+	FSectionID()
+	{
+	}
+
+	explicit FSectionID( const FElementID InitElementID )
+		: FElementID( InitElementID.GetValue() )
+	{
+	}
+
+	explicit FSectionID( const int32 InitIDValue )
+		: FElementID( InitIDValue )
+	{
+	}
+
+	/** Invalid section ID */
+	MESHEDITINGRUNTIME_API static const FSectionID Invalid;
+};
+
+
+USTRUCT( BlueprintType )
+struct FPolygonID : public FElementID
+{
+	GENERATED_BODY()
+
+	FPolygonID()
+	{
+	}
+
+	explicit FPolygonID( const FElementID InitElementID )
+		: FElementID( InitElementID.GetValue() )
+	{
+	}
+
+	explicit FPolygonID( const int32 InitIDValue )
+		: FElementID( InitIDValue )
+	{
+	}
+
+	/** Invalid polygon ID */
+	MESHEDITINGRUNTIME_API static const FPolygonID Invalid;	// @todo mesheditor script: Can we expose these to BP nicely?	Do we even need to?
+};
+
+
+USTRUCT()
+struct FTriangleID : public FElementID
+{
+	GENERATED_BODY()
+
+	FTriangleID()
+	{
+	}
+
+	explicit FTriangleID( const FElementID InitElementID )
+		: FElementID( InitElementID.GetValue() )
+	{
+	}
+
+	explicit FTriangleID( const uint32 InitIDValue )
+		: FElementID( InitIDValue )
+	{
+	}
+
+	/** Invalid triangle ID */
+	MESHEDITINGRUNTIME_API static const FTriangleID Invalid;
+};
+
+
+/**
+ * Unique identifies a specific sub-mesh within a component
+ */
+struct FEditableMeshSubMeshAddress
+{
+	/** Pointer that uniquely identifies the mesh object being edited (not the instance), for hashing/comparison purposes */
+	void* MeshObjectPtr;
+
+	/** The type of mesh format */
+	// @todo mesheditor: Should SHOULD be OK, but if it somehow was deallocated while we were using it, we'd want to store an FName or ID here instead of a pointer.  Should be fine though.
+	class IEditableMeshFormat* EditableMeshFormat;
+
+	/** The index of the mesh within the component, for components that may define more than one mesh */
+	int32 MeshIndex;
+
+	/** The mesh level of detail index, or zero if not applicable to the type of mesh. */
+	int32 LODIndex;
+
+
+	/** Default constructor that initializes variables to an default sub-mesh address */
+	FEditableMeshSubMeshAddress()
+		: MeshObjectPtr( nullptr ),
+		  EditableMeshFormat( nullptr ),
+		  MeshIndex( 0 ),
+		  LODIndex( 0 )
+	{
+	}
+
+	/** Equality check */
+	inline bool operator==( const FEditableMeshSubMeshAddress& Other ) const
+	{
+		return
+			MeshObjectPtr == Other.MeshObjectPtr &&
+			EditableMeshFormat == Other.EditableMeshFormat &&
+			MeshIndex == Other.MeshIndex &&
+			LODIndex == Other.LODIndex;
+	}
+
+	/** Hashing */
+	FORCEINLINE friend uint32 GetTypeHash( const FEditableMeshSubMeshAddress& Other )
+	{
+		// @todo mesheditor: Hash could be improved a bit to consider LOD/MeshIndex, etc.
+		return GetTypeHash( Other.MeshObjectPtr );
+	}
+
+	/** Convert to a string */
+	inline FString ToString() const
+	{
+		return FString::Printf(
+			TEXT( "PtrHash:%lu, FmtHash:%lu, MeshIndex:%lu, LODIndex:%lu" ),
+			GetTypeHash( MeshObjectPtr ),
+			GetTypeHash( EditableMeshFormat ),
+			MeshIndex,
+			LODIndex );
+	}
+};
+
+
+
+USTRUCT( BlueprintType )
+struct FPolygonRef
+{
+	GENERATED_BODY()
+
+	UPROPERTY( BlueprintReadWrite, Category="Editable Mesh" )
+	FSectionID SectionID;
+
+	UPROPERTY( BlueprintReadWrite, Category="Editable Mesh" )
+	FPolygonID PolygonID;
+
+	FPolygonRef()
+	{
+	}
+
+	FPolygonRef( const FSectionID InitSectionID, const FPolygonID InitPolygonID )
+		: SectionID( InitSectionID ),
+		  PolygonID( InitPolygonID )
+	{
+	}
+
+	inline bool operator==( const FPolygonRef& Other ) const
+	{
+		return SectionID == Other.SectionID && PolygonID == Other.PolygonID;
+	}
+
+	inline bool operator!=( const FPolygonRef& Other ) const
+	{
+		return SectionID != Other.SectionID || PolygonID != Other.PolygonID;
+	}
+
+	/** Hashing */
+	FORCEINLINE friend uint32 GetTypeHash( const FPolygonRef& Other )
+	{
+		return GetTypeHash( ( uint64( Other.PolygonID.GetValue() ) << 32 ) | uint64( Other.SectionID.GetValue() ) );
+	}
+
+	FString ToString() const
+	{
+		return FString::Printf( TEXT( "SectionID:%s, PolygonID:%s" ), *SectionID.ToString(), *PolygonID.ToString() );
+	}
+
+	friend FArchive& operator<<( FArchive& Ar, FPolygonRef& PolygonRef )
+	{
+		Ar << PolygonRef.SectionID;
+		Ar << PolygonRef.PolygonID;
+		return Ar;
+	}
+
+	/** Invalid polygon ref */
+	MESHEDITINGRUNTIME_API static const FPolygonRef Invalid;
+};
+
+
+UCLASS( abstract )
+class MESHEDITINGRUNTIME_API UEditableMeshAttribute : public UObject
+{
+	GENERATED_BODY()
+
+public:
+
+	//
+	// Vertex data for any vertex
+	//
+
+	/** Static: The attribute name for vertex position */
+	UFUNCTION( BlueprintPure, Category="Editable Mesh" ) static inline FName VertexPosition()
+	{
+		return VertexPositionName;
+	}
+
+	/** Static: The attribute name for vertex corner sharpness (only applies to subdivision meshes) */
+	UFUNCTION( BlueprintPure, Category="Editable Mesh" ) static inline FName VertexCornerSharpness()
+	{
+		return VertexCornerSharpnessName;
+	}
+
+	//
+	// Polygon-specific vertex data (can also be set on the vertex itself to update all polygon vertices.)
+	//
+
+	/** Static: The attribute name for vertex normal (tangent Z) */
+	UFUNCTION( BlueprintPure, Category="Editable Mesh" ) static inline FName VertexNormal()
+	{
+		return VertexNormalName;
+	}
+
+	/** Static: The attribute name for vertex tangent vector (tangent X) */
+	UFUNCTION( BlueprintPure, Category="Editable Mesh" ) static inline FName VertexTangent()
+	{
+		return VertexTangentName;
+	}
+
+	/** Static: The attribute name for the vertex basis determinant sign (used to calculate the direction of tangent Y) */
+	UFUNCTION( BlueprintPure, Category="Editable Mesh" ) static inline FName VertexBinormalSign()
+	{
+		return VertexBinormalSignName;
+	}
+
+	/** Static: The attribute name for vertex texture coordinate.  The attribute index defines which texture coordinate set. */
+	UFUNCTION( BlueprintPure, Category="Editable Mesh" ) static inline FName VertexTextureCoordinate()
+	{
+		return VertexTextureCoordinateName;
+	}
+
+	/** Static: The attribute name for the vertex color. */
+	UFUNCTION( BlueprintPure, Category="Editable Mesh" ) static inline FName VertexColor()
+	{
+		return VertexColorName;
+	}
+
+	//
+	// Edges
+	//
+
+	/** Static: The attribute name for edge hardedness */
+	UFUNCTION( BlueprintPure, Category="Editable Mesh" ) static inline FName EdgeIsHard()
+	{
+		return EdgeIsHardName;
+	}
+
+	/** Static: The attribute name for edge crease sharpness (only applies to subdivision meshes) */
+	UFUNCTION( BlueprintPure, Category="Editable Mesh" ) static inline FName EdgeCreaseSharpness()
+	{
+		return EdgeCreaseSharpnessName;
+	}
+
+private:
+
+	static const FName VertexPositionName;
+	static const FName VertexCornerSharpnessName;
+	static const FName VertexNormalName;
+	static const FName VertexTangentName;
+	static const FName VertexBinormalSignName;
+	static const FName VertexTextureCoordinateName;
+	static const FName VertexColorName;
+	static const FName EdgeIsHardName;
+	static const FName EdgeCreaseSharpnessName;
+};
+
+
+UENUM()
+enum class EMeshModificationType : uint8
+{
+	/** The first Interim change since the last Final change.  This must be followed by either an Interim change or a Final change */
+	FirstInterim,
+
+	/** User is still in the middle of their interaction.  More changes to come, so don't bother finalizing everything yet (smoother performance) */
+	Interim,
+
+	/** User has finished their current interaction with this mesh, and everything needs to be finalized at this time */
+	Final
+};
+
+
+UENUM()
+enum class EMeshTopologyChange : uint8
+{
+	/** We won't be changing the mesh topology, but values could be changed (vertex positions, UVs, colors, etc.) */
+	NoTopologyChange,
+
+	/** Topology is changing with this edit, potentially along with other changes */
+	TopologyChange
+};
+
+
+
+USTRUCT( BlueprintType )
+struct FSubdividedQuadVertex
+{
+	GENERATED_BODY()
+
+	/** The index of the vertex position (into the FSubdivisionLimitData's VertexPositions array) used for this vertex */
+	UPROPERTY( BlueprintReadWrite, Category="Editable Mesh" )	int32 VertexPositionIndex;
+
+	/** Texture coordinates for this vertex.  We only support up to two, for now. (Just to avoid TArrays/allocations) */
+	UPROPERTY( BlueprintReadWrite, Category="Editable Mesh" )	FVector2D TextureCoordinate0;
+	UPROPERTY( BlueprintReadWrite, Category="Editable Mesh" )	FVector2D TextureCoordinate1;
+
+	/** Vertex color */
+	UPROPERTY( BlueprintReadWrite, Category="Editable Mesh" )	FColor VertexColor;
+
+	/** Quad vertex normal */
+	UPROPERTY( BlueprintReadWrite, Category="Editable Mesh" )	FVector VertexNormal;
+
+	/** Quad vertex tangent */
+	UPROPERTY( BlueprintReadWrite, Category="Editable Mesh" )	FVector VertexTangent;
+
+	/** Quad vertex binormal sign (-1.0 or 1.0)*/
+	UPROPERTY( BlueprintReadWrite, Category="Editable Mesh" )	float VertexBinormalSign;
+};
+
+
+USTRUCT( BlueprintType )
+struct FSubdividedQuad
+{
+	GENERATED_BODY()
+
+	// NOTE: The reason we're using separate variables instead of a static array is so that we can expose these to Blueprints, which doesn't support static array properties
+
+	/** The vertices for the four corners of this quad */
+	UPROPERTY( BlueprintReadWrite, Category="Editable Mesh" )	FSubdividedQuadVertex QuadVertex0;
+	UPROPERTY( BlueprintReadWrite, Category="Editable Mesh" )	FSubdividedQuadVertex QuadVertex1;
+	UPROPERTY( BlueprintReadWrite, Category="Editable Mesh" )	FSubdividedQuadVertex QuadVertex2;
+	UPROPERTY( BlueprintReadWrite, Category="Editable Mesh" )	FSubdividedQuadVertex QuadVertex3;
+
+	const FSubdividedQuadVertex& GetQuadVertex( const int32 Index ) const
+	{
+		switch( Index )
+		{
+			default:
+			case 0:
+				return QuadVertex0;
+
+			case 1:
+				return QuadVertex1;
+
+			case 2:
+				return QuadVertex2;
+			
+			case 3:
+				return QuadVertex3;
+		}
+	}
+
+	FSubdividedQuadVertex& AccessQuadVertex( const int32 Index )
+	{
+		switch( Index )
+		{
+			default:
+			case 0:
+				return QuadVertex0;
+
+			case 1:
+				return QuadVertex1;
+
+			case 2:
+				return QuadVertex2;
+
+			case 3:
+				return QuadVertex3;
+		}
+	}
+};
+
+
+USTRUCT( BlueprintType )
+struct FSubdividedWireEdge
+{
+	GENERATED_BODY()
+
+	// NOTE: The reason we're using separate variables instead of a static array is so that we can expose these to Blueprints, which doesn't support static array properties
+
+	/** The vertex indices for the two corners of this quad */
+	UPROPERTY( BlueprintReadWrite, Category="Editable Mesh" )	int32 EdgeVertex0PositionIndex;
+	UPROPERTY( BlueprintReadWrite, Category="Editable Mesh" )	int32 EdgeVertex1PositionIndex;
+
+	/** True if this edge is a counterpart to an original base cage edge of the mesh.  Otherwise it's a new edge that exists only
+	    in the subdivision surfaces */
+	bool bIsBaseCageCounterpartEdge;
+};
+
+
+USTRUCT( BlueprintType )
+struct FSubdivisionLimitSection
+{
+	GENERATED_BODY()
+
+	/** All of the quads in this section, as a result from subdividing the mesh */
+	UPROPERTY( BlueprintReadWrite, Category="Editable Mesh" )
+	TArray<FSubdividedQuad> SubdividedQuads;
+};
+
+
+USTRUCT( BlueprintType )
+struct FSubdivisionLimitData
+{
+	GENERATED_BODY()
+
+	/** Positions of all of the vertices for this subdivision level.  Many vertex positions may be shared between subdivided quads. */
+	UPROPERTY( BlueprintReadWrite, Category="Editable Mesh" )
+	TArray<FVector> VertexPositions;
+
+	/** Data for each of the sections in the mesh.  This array will have the same number of elements as the editable mesh's 
+	    section list (not necessarily the same indices though, due to sparseness). */
+	UPROPERTY( BlueprintReadWrite, Category="Editable Mesh" )
+	TArray<FSubdivisionLimitSection> Sections;	
+
+	/** All of the wire edges in the entire mesh (for all sections) */
+	UPROPERTY( BlueprintReadWrite, Category="Editable Mesh" )
+	TArray<FSubdividedWireEdge> SubdividedWireEdges;
+};
+
+
+USTRUCT( BlueprintType )
+struct FMeshElementAttributeData
+{
+	GENERATED_BODY()
+
+	/** Name of the attribute */
+	UPROPERTY( BlueprintReadWrite, Category="Editable Mesh" )
+	FName AttributeName;
+
+	/** Index of the attribute */
+	UPROPERTY( BlueprintReadWrite, Category="Editable Mesh" )
+	int32 AttributeIndex;
+
+	/** The value of this attribute */
+	UPROPERTY( BlueprintReadWrite, Category="Editable Mesh" )
+	FVector4 AttributeValue;
+
+	/** Default constructor */
+	FMeshElementAttributeData()
+		: AttributeName( NAME_None ),
+		  AttributeIndex( 0 ),
+		  AttributeValue( 0.0f )
+	{
+	}
+
+	/** Constructor that fills in all fields */
+	FMeshElementAttributeData( const FName InitAttributeName, const int32 InitAttributeIndex, const FVector4 InitAttributeValue )
+		: AttributeName( InitAttributeName ),
+		  AttributeIndex( InitAttributeIndex ),
+		  AttributeValue( InitAttributeValue )
+	{
+	}
+
+	FString ToString() const
+	{
+		return FString::Printf(
+			TEXT( "Name:%s, Index:%lu, Value:%s" ),
+			*AttributeName.ToString(),
+			AttributeIndex,
+			*AttributeValue.ToString() );
+	}
+};
+
+
+USTRUCT( BlueprintType )
+struct FMeshElementAttributeList
+{
+	GENERATED_BODY()
+
+	/** List of attributes to apply to a mesh element */
+	UPROPERTY( BlueprintReadWrite, Category="Editable Mesh" )
+	TArray<FMeshElementAttributeData> Attributes;
+
+	/** Default constructor */
+	FMeshElementAttributeList()
+		: Attributes()
+	{
+	}
+
+	FString ToString() const
+	{
+		return FString::Printf(
+			TEXT( "Attributes:%s" ),
+			*LogHelpers::ArrayToString( Attributes ) );
+	}
+};
+
+
+USTRUCT( BlueprintType )
+struct FVertexToCreate
+{
+	GENERATED_BODY()
+
+	/** Attributes of this vertex itself */
+	UPROPERTY( BlueprintReadWrite, Category="Editable Mesh" )
+	FMeshElementAttributeList VertexAttributes;
+
+	/** The original ID of the vertex.  Should only be used by the undo system. */
+	UPROPERTY()
+	FVertexID OriginalVertexID;
+
+	FVertexToCreate()
+		: VertexAttributes(),
+		  OriginalVertexID( FVertexID::Invalid )
+	{
+	}
+
+	FString ToString() const
+	{
+		return FString::Printf(
+			TEXT( "VertexAttributes:%s, OriginalVertexID:%s" ),
+			*VertexAttributes.ToString(),
+			*OriginalVertexID.ToString() );
+	}
+};
+
+
+USTRUCT( BlueprintType )
+struct FEdgeToCreate
+{
+	GENERATED_BODY()
+		
+	/** The first vertex this edge connects */
+	UPROPERTY( BlueprintReadWrite, Category="Editable Mesh" )
+	FVertexID VertexID0;
+
+	/** The second vertex this edge connects */
+	UPROPERTY( BlueprintReadWrite, Category="Editable Mesh" )
+	FVertexID VertexID1;
+
+	/** The polygons that are connected to this edge. */
+	UPROPERTY( BlueprintReadWrite, Category="Editable Mesh" )
+	TArray<FPolygonRef> ConnectedPolygons;
+
+	/** Attributes of this edge itself */
+	UPROPERTY( BlueprintReadWrite, Category="Editable Mesh" )
+	FMeshElementAttributeList EdgeAttributes;
+
+	/** The original ID of the edge.  Should only be used by the undo system. */
+	UPROPERTY( BlueprintReadWrite, Category="Editable Mesh" )
+	FEdgeID OriginalEdgeID;
+
+	/** Default constructor */
+	FEdgeToCreate()
+		: VertexID0( FVertexID::Invalid ),
+		  VertexID1( FVertexID::Invalid ),
+		  ConnectedPolygons(),
+		  EdgeAttributes(),
+		  OriginalEdgeID( FEdgeID::Invalid )
+	{
+	}
+
+	FString ToString() const
+	{
+		return FString::Printf(
+			TEXT( "VertexID0:%s, VertexID1:%s, ConnectedPolygons:%s, OriginalEdgeID:%s" ),
+			*VertexID0.ToString(), 
+			*VertexID1.ToString(),
+			*LogHelpers::ArrayToString( ConnectedPolygons ),
+			*EdgeAttributes.ToString(),
+			OriginalEdgeID.GetValue() );
+	}
+};
+
+
+
+USTRUCT( BlueprintType )
+struct FVertexAndAttributes
+{
+	GENERATED_BODY()
+
+	/** The vertex ID to insert into the polygon */
+	UPROPERTY( BlueprintReadWrite, Category="Editable Mesh" )
+	FVertexID VertexID;
+
+	/** A list of polygon attributes to set for the vertex on the polygon we're inserting it into */
+	UPROPERTY( BlueprintReadWrite, Category="Editable Mesh" )
+	FMeshElementAttributeList PolygonVertexAttributes;
+
+	/** Default constructor */
+	FVertexAndAttributes()
+		: VertexID( 0 ),
+		  PolygonVertexAttributes()
+	{
+	}
+
+	FString ToString() const
+	{
+		return FString::Printf(
+			TEXT( "VertexID:%s, PolygonVertexAttributes:%s" ),
+			*VertexID.ToString(),
+			*PolygonVertexAttributes.ToString() );
+	}
+};
+
+
+USTRUCT( BlueprintType )
+struct FPolygonHoleVertices
+{
+	GENERATED_BODY()
+
+	/** Ordered list of vertices that defines the hole's contour, along with the polygon vertex attributes to set for each vertex */
+	UPROPERTY( BlueprintReadWrite, Category="Editable Mesh" )
+	TArray<FVertexAndAttributes> HoleVertices;
+
+	/** Default constructor */
+	FPolygonHoleVertices()
+		: HoleVertices()
+	{
+	}
+
+	FString ToString() const
+	{
+		return FString::Printf(
+			TEXT( "HoleVertices:%s" ),
+			*LogHelpers::ArrayToString( HoleVertices ) );
+	}
+};
+
+
+
+USTRUCT( BlueprintType )
+struct FPolygonToCreate
+{
+	GENERATED_BODY()
+
+	/** The section the polygon will be added to */
+	UPROPERTY( BlueprintReadWrite, Category="Editable Mesh" )
+	FSectionID SectionID;
+
+	/** Ordered list of vertices that defines the polygon's perimeter, along with the polygon vertex attributes to set for each vertex */
+	UPROPERTY( BlueprintReadWrite, Category="Editable Mesh" )
+	TArray<FVertexAndAttributes> PerimeterVertices;
+	
+	/** For each hole in the polygon, an ordered list of vertices that defines that hole's boundary */
+	UPROPERTY( BlueprintReadWrite, Category="Editable Mesh" )
+	TArray<FPolygonHoleVertices> PolygonHoles;
+
+	/** The original ID of the polygon.  Should only be used by the undo system. */
+	UPROPERTY( BlueprintReadWrite, Category="Editable Mesh" )
+	FPolygonID OriginalPolygonID;
+
+	/** Default constructor */
+	FPolygonToCreate()
+		: SectionID( 0 ),
+		  PerimeterVertices(),
+		  PolygonHoles(),
+		  OriginalPolygonID( FPolygonID::Invalid )
+	{
+	}
+
+	FString ToString() const
+	{
+		return FString::Printf(
+			TEXT( "SectionID:%s, PerimeterVertices:%s, PolygonHoles:%s, OriginalPolygonID:%s" ),
+			*SectionID.ToString(),
+			*LogHelpers::ArrayToString( PerimeterVertices ),
+			*LogHelpers::ArrayToString( PolygonHoles ),
+			*OriginalPolygonID.ToString() );
+	}
+};
+
+
+
+USTRUCT( BlueprintType )
+struct FVertexPair
+{
+	GENERATED_BODY()
+
+	/** The first vertex ID in this pair */
+	UPROPERTY( BlueprintReadWrite, Category="Editable Mesh" )
+	FVertexID VertexID0;
+
+	/** The second vertex ID in this pair */
+	UPROPERTY( BlueprintReadWrite, Category="Editable Mesh" )
+	FVertexID VertexID1;
+
+	/** Default constructor */
+	FVertexPair()
+		: VertexID0( FVertexID::Invalid ),
+		  VertexID1( FVertexID::Invalid )
+	{
+	}
+
+	// @todo mesheditor: Do we need a ToString() function?  Would it ever be called?
+};
+
+
+USTRUCT( BlueprintType )
+struct FPolygonToSplit
+{
+	GENERATED_BODY()
+
+	/** The polygon that we'll be splitting */
+	UPROPERTY( BlueprintReadWrite, Category="Editable Mesh" )
+	FPolygonRef PolygonRef;
+
+	/** A list of pairs of vertices that new edges will be created at.  The pairs must be ordered, and the vertices
+	    must already exist and be connected to the polygon */
+	UPROPERTY( BlueprintReadWrite, Category="Editable Mesh" )
+	TArray<FVertexPair> VertexPairsToSplitAt;
+
+	/** Default constructor */
+	FPolygonToSplit()
+		: PolygonRef( FSectionID::Invalid, FPolygonID::Invalid ),
+		  VertexPairsToSplitAt()
+	{
+	}
+
+	// @todo mesheditor: Do we need a ToString() function?  Would it ever be called?
+};
+
+
+USTRUCT( BlueprintType )
+struct FAttributesForVertex
+{
+	GENERATED_BODY()
+
+	/** The vertex ID to set attributes on */
+	UPROPERTY( BlueprintReadWrite, Category="Editable Mesh" )
+	FVertexID VertexID;
+
+	/** A list of attributes to set for the vertex */
+	UPROPERTY( BlueprintReadWrite, Category="Editable Mesh" )
+	FMeshElementAttributeList VertexAttributes;
+
+	/** Default constructor */
+	FAttributesForVertex()
+		: VertexID( 0 ),
+		  VertexAttributes()
+	{
+	}
+
+	FString ToString() const
+	{
+		return FString::Printf(
+			TEXT( "VertexID:%s, VertexAttributes:%s" ),
+			*VertexID.ToString(),
+			*VertexAttributes.ToString() );
+	}
+};
+
+
+USTRUCT( BlueprintType )
+struct FAttributesForEdge
+{
+	GENERATED_BODY()
+
+	/** The edge ID to set attributes on */
+	UPROPERTY( BlueprintReadWrite, Category="Editable Mesh" )
+	FEdgeID EdgeID;
+
+	/** A list of attributes to set for the edge */
+	UPROPERTY( BlueprintReadWrite, Category="Editable Mesh" )
+	FMeshElementAttributeList EdgeAttributes;
+
+	/** Default constructor */
+	FAttributesForEdge()
+		: EdgeID( 0 ),
+		  EdgeAttributes()
+	{
+	}
+
+	FString ToString() const
+	{
+		return FString::Printf(
+			TEXT( "EdgeID:%s, EdgeAttributes:%s" ),
+			*EdgeID.ToString(),
+			*EdgeAttributes.ToString() );
+	}
+};
+
+
+USTRUCT( BlueprintType )
+struct FVertexAttributesForPolygonHole
+{
+	GENERATED_BODY()
+
+	/** For each hole vertex, a list of attributes for that vertex.  You can leave a given array empty for 
+	    a specific hole index if you don't want to set attributes for select holes. */
+	UPROPERTY( BlueprintReadWrite, Category="Editable Mesh" )
+	TArray<FMeshElementAttributeList> VertexAttributeList;
+
+	/** Default constructor */
+	FVertexAttributesForPolygonHole()
+		: VertexAttributeList()
+	{
+	}
+
+	FString ToString() const
+	{
+		return FString::Printf(
+			TEXT( "VertexAttributeList:%s" ),
+			*LogHelpers::ArrayToString( VertexAttributeList ) );
+	}
+};
+
+
+USTRUCT( BlueprintType )
+struct FVertexAttributesForPolygon
+{
+	GENERATED_BODY()
+
+	/** The polygon to set vertex attributes on */
+	UPROPERTY( BlueprintReadWrite, Category="Editable Mesh" )
+	FPolygonRef PolygonRef;
+
+	/** For each polygon vertex, a list of attributes for that vertex.  Can be left empty if you don't want to set any attributes. */
+	UPROPERTY( BlueprintReadWrite, Category="Editable Mesh" )
+	TArray<FMeshElementAttributeList> PerimeterVertexAttributeLists;
+
+	/** For each hole vertex, a list of attributes for that vertex.  Can be left empty if you don't want to set any attributes.  Also
+	    you can leave a given array empty for a specific hole index if you don't want to set attributes for select holes. */
+	UPROPERTY( BlueprintReadWrite, Category="Editable Mesh" )
+	TArray<FVertexAttributesForPolygonHole> VertexAttributeListsForEachHole;
+
+	/** Default constructor */
+	FVertexAttributesForPolygon()
+		: PolygonRef( FSectionID::Invalid, FPolygonID::Invalid ),
+		  PerimeterVertexAttributeLists(),
+		  VertexAttributeListsForEachHole()
+	{
+	}
+
+	FString ToString() const
+	{
+		return FString::Printf(
+			TEXT( "PolygonRef:%s, PerimeterVertexAttributeLists:%s, VertexAttributeListsForEachHole:%s" ),
+			*PolygonRef.ToString(),
+			*LogHelpers::ArrayToString( PerimeterVertexAttributeLists ),
+			*LogHelpers::ArrayToString( VertexAttributeListsForEachHole ) );
+	}
+};
+
+
+USTRUCT( BlueprintType )
+struct FVerticesForEdge
+{
+	GENERATED_BODY()
+
+	/** The edge ID */
+	UPROPERTY( BlueprintReadWrite, Category="Editable Mesh" )
+	FEdgeID EdgeID;
+
+	/** First new vertex ID for this edge */
+	UPROPERTY( BlueprintReadWrite, Category="Editable Mesh" )
+	FVertexID NewVertexID0;
+
+	/** Second new vertex ID for this edge */
+	UPROPERTY( BlueprintReadWrite, Category="Editable Mesh" )
+	FVertexID NewVertexID1;
+
+	/** Default constructor */
+	FVerticesForEdge()
+		: EdgeID( 0 ),
+		  NewVertexID0( FVertexID::Invalid ),	// @todo mesheditor urgent: Typesafety isn't working -- able to assign EdgeID to VertexID!
+		  NewVertexID1( FVertexID::Invalid )
+	{
+	}
+
+	FString ToString() const
+	{
+		return FString::Printf(
+			TEXT( "EdgeID:%s, NewVertexID0:%s, NewVertexID1:%s]" ),
+			*EdgeID.ToString(),
+			*NewVertexID0.ToString(), 
+			*NewVertexID1.ToString() );
+	}
+};
+
+
+USTRUCT( BlueprintType )
+struct FVertexToMove
+{
+	GENERATED_BODY()
+
+	/** The vertex we'll be moving around */
+	UPROPERTY( BlueprintReadWrite, Category="Editable Mesh" )
+	FVertexID VertexID;
+
+	/** The new position of the vertex */
+	UPROPERTY( BlueprintReadWrite, Category="Editable Mesh" )
+	FVector NewVertexPosition;
+
+
+	FVertexToMove()
+		: VertexID( FVertexID::Invalid ),
+		  NewVertexPosition( FVector::ZeroVector )
+	{
+	}
+
+	FString ToString() const
+	{
+		return FString::Printf(
+			TEXT( "VertexID:%s, NewVertexPosition:%s" ),
+			*VertexID.ToString(),
+			*NewVertexPosition.ToString() );
+	}
+};
+
+
+USTRUCT( BlueprintType )
+struct FSectionToCreate
+{
+	GENERATED_BODY()
+
+	/** Material to assign to the new section */
+	UPROPERTY( BlueprintReadWrite, Category="Editable Mesh" )
+	UMaterialInterface* Material;
+
+	/** Whether the new section should have collision enabled */
+	UPROPERTY( BlueprintReadWrite, Category="Editable Mesh" )
+	bool bEnableCollision;
+
+	/** Whether the new section casts a shadow */
+	UPROPERTY( BlueprintReadWrite, Category="Editable Mesh" )
+	bool bCastShadow;
+
+	/** The original ID of the section.  Should only be used by the undo system. */
+	UPROPERTY( BlueprintReadWrite, Category="Editable Mesh" )
+	FSectionID OriginalSectionID;
+
+	/** The original rendering section index.  Should only be used by the undo system. */
+	UPROPERTY( BlueprintReadWrite, Category="Editable Mesh" )
+	int32 OriginalRenderingSectionIndex;
+
+	/** Default constructor */
+	FSectionToCreate()
+		: Material( nullptr )
+		, bEnableCollision( false )
+		, bCastShadow( false )
+		, OriginalSectionID( FSectionID::Invalid )
+		, OriginalRenderingSectionIndex( INDEX_NONE )
+	{
+	}
+
+	FString ToString() const
+	{
+		return FString::Printf(
+			TEXT( "Material:%s, bEnableCollision:%s, bCastShadow:%s, OriginalSectionID:%s" ),
+			Material ? *Material->GetName() : TEXT( "<none>" ),
+			*LexicalConversion::ToString( bEnableCollision ),
+			*LexicalConversion::ToString( bCastShadow ),
+			*OriginalSectionID.ToString() );
+	}
+};
diff --git a/Engine/Source/Runtime/MeshEditingRuntime/Public/IEditableMeshFormat.h b/Engine/Source/Runtime/MeshEditingRuntime/Public/IEditableMeshFormat.h
new file mode 100644
--- /dev/null
+++ b/Engine/Source/Runtime/MeshEditingRuntime/Public/IEditableMeshFormat.h
@@ -0,0 +1,22 @@
+// Copyright 1998-2015 Epic Games, Inc. All Rights Reserved.
+
+#pragma once
+
+#include "CoreMinimal.h"
+#include "Features/IModularFeature.h"
+#include "Features/IModularFeatures.h"
+
+
+/**
+ * Interface for mesh primitives that can be interactively edited (low-poly editing, sculpting, attribute painting, etc.)
+ */
+class IEditableMeshFormat : public IModularFeature
+{
+
+public:
+
+	// @todo mesheditor: Comments!
+	virtual void FillMeshObjectPtr( class UPrimitiveComponent& Component, struct FEditableMeshSubMeshAddress& SubMeshAddress ) = 0;
+	virtual class UEditableMesh* MakeEditableMesh( class UPrimitiveComponent& Component, const struct FEditableMeshSubMeshAddress& SubMeshAddress ) = 0;
+
+};
\ No newline at end of file
diff --git a/Engine/Source/Runtime/MeshEditingRuntime/Public/IMeshEditingRuntimeModule.h b/Engine/Source/Runtime/MeshEditingRuntime/Public/IMeshEditingRuntimeModule.h
new file mode 100644
--- /dev/null
+++ b/Engine/Source/Runtime/MeshEditingRuntime/Public/IMeshEditingRuntimeModule.h
@@ -0,0 +1,39 @@
+// Copyright 1998-2015 Epic Games, Inc. All Rights Reserved.
+
+#pragma once
+
+#include "CoreMinimal.h"
+#include "ModuleManager.h"
+
+
+/**
+ * The public interface to this module
+ */
+class IMeshEditingRuntimeModule : public IModuleInterface
+{
+
+public:
+
+	/**
+	 * Singleton-like access to this module's interface.  This is just for convenience!
+	 * Beware of calling this during the shutdown phase, though.  Your module might have been unloaded already.
+	 *
+	 * @return Returns singleton instance, loading the module on demand if needed
+	 */
+	static inline IMeshEditingRuntimeModule& Get()
+	{
+		return FModuleManager::LoadModuleChecked< IMeshEditingRuntimeModule >( "MeshEditingRuntime" );
+	}
+
+	/**
+	 * Checks to see if this module is loaded and ready.  It is only valid to call Get() if IsAvailable() returns true.
+	 *
+	 * @return True if the module is loaded and ready to use
+	 */
+	static inline bool IsAvailable()
+	{
+		return FModuleManager::Get().IsModuleLoaded( "MeshEditingRuntime" );
+	}
+
+};
+
diff --git a/Engine/Source/Runtime/MeshEditingRuntime/StaticMeshEditableMeshFormat.cpp b/Engine/Source/Runtime/MeshEditingRuntime/StaticMeshEditableMeshFormat.cpp
new file mode 100644
--- /dev/null
+++ b/Engine/Source/Runtime/MeshEditingRuntime/StaticMeshEditableMeshFormat.cpp
@@ -0,0 +1,60 @@
+// Copyright 1998-2015 Epic Games, Inc. All Rights Reserved.
+
+#include "StaticMeshEditableMeshFormat.h"
+#include "EditableMesh.h"
+#include "Engine/StaticMesh.h"
+#include "Components/StaticMeshComponent.h"
+#include "StaticMeshResources.h"
+#include "EditableStaticMesh.h"
+
+
+void FStaticMeshEditableMeshFormat::FillMeshObjectPtr( UPrimitiveComponent& Component, FEditableMeshSubMeshAddress& SubMeshAddress )
+{
+	SubMeshAddress.MeshObjectPtr = nullptr;
+
+	const UStaticMeshComponent* StaticMeshComponentPtr = Cast<const UStaticMeshComponent>( &Component );
+	if( StaticMeshComponentPtr != nullptr )
+	{
+		const UStaticMeshComponent& StaticMeshComponent = *StaticMeshComponentPtr;
+
+		UStaticMesh* ComponentStaticMesh = StaticMeshComponent.GetStaticMesh();
+		if( ComponentStaticMesh != nullptr && ComponentStaticMesh->HasValidRenderData() )
+		{
+			SubMeshAddress.MeshObjectPtr = ComponentStaticMesh;
+		}
+	}
+}
+
+
+UEditableMesh* FStaticMeshEditableMeshFormat::MakeEditableMesh( UPrimitiveComponent& Component, const FEditableMeshSubMeshAddress& SubMeshAddress )
+{
+	// If the static mesh already has an attached UEditableStaticMesh, use that in preference to creating a new one
+	const UStaticMeshComponent* StaticMeshComponentPtr = Cast<const UStaticMeshComponent>( &Component );
+	if( StaticMeshComponentPtr != nullptr )
+	{
+		const UStaticMeshComponent& StaticMeshComponent = *StaticMeshComponentPtr;
+		if( StaticMeshComponent.GetStaticMesh() != nullptr )
+		{
+			UEditableMesh* EditableMesh = StaticMeshComponent.GetStaticMesh()->EditableMesh;
+			if( EditableMesh )
+			{
+				EditableMesh->SetSubMeshAddress( SubMeshAddress );
+				return EditableMesh;
+			}
+		}
+	}
+
+	UEditableStaticMesh* EditableStaticMesh = NewObject<UEditableStaticMesh>();
+
+	EditableStaticMesh->InitEditableStaticMesh( Component, SubMeshAddress );
+
+	// Don't bother returning a new mesh if it has no geometry
+	if( EditableStaticMesh->GetVertexCount() == 0 )
+	{
+		EditableStaticMesh->MarkPendingKill();
+		EditableStaticMesh = nullptr;
+	}
+
+	return EditableStaticMesh;
+}
+
diff --git a/Engine/Source/Runtime/MeshEditingRuntime/StaticMeshEditableMeshFormat.h b/Engine/Source/Runtime/MeshEditingRuntime/StaticMeshEditableMeshFormat.h
new file mode 100644
--- /dev/null
+++ b/Engine/Source/Runtime/MeshEditingRuntime/StaticMeshEditableMeshFormat.h
@@ -0,0 +1,21 @@
+// Copyright 1998-2015 Epic Games, Inc. All Rights Reserved.
+
+#pragma once
+
+#include "CoreMinimal.h"
+#include "IEditableMeshFormat.h"
+
+
+/**
+ * Implements interactive mesh editing support for Static Meshes
+ */ 
+class FStaticMeshEditableMeshFormat : public IEditableMeshFormat
+{
+
+public:
+
+	// IEditableMeshFormat interface
+	virtual void FillMeshObjectPtr( class UPrimitiveComponent& Component, FEditableMeshSubMeshAddress& SubMeshAddress ) override;
+	virtual UEditableMesh* MakeEditableMesh( class UPrimitiveComponent& Component, const FEditableMeshSubMeshAddress& SubMeshAddress ) override;
+
+};
\ No newline at end of file
