Mesh Editor: Show selected polygons and edges on subdivision preview mesh
- If you have polygons or edges selected on a base cage, their counterparts will be highlighted on the subdivision preview wire overlay
- Added new Editable Mesh function to get the edge that is connecting two specified vertices (UEditableMesh::GetEdgeThatConnectsVertices)

diff --git a/Engine/Plugins/Editor/MeshEditor/Source/MeshEditor/MeshEditorMode.cpp b/Engine/Plugins/Editor/MeshEditor/Source/MeshEditor/MeshEditorMode.cpp
--- a/Engine/Plugins/Editor/MeshEditor/Source/MeshEditor/MeshEditorMode.cpp
+++ b/Engine/Plugins/Editor/MeshEditor/Source/MeshEditor/MeshEditorMode.cpp
@@ -2247,9 +2247,41 @@ void FMeshEditorMode::Render( const FSceneView* SceneView, FViewport* Viewport,
 			// Draw subdivision limit surface if subdivision preview is enabled
 			if( EditableMesh->IsPreviewingSubdivisions() )
 			{
+				// Figure out all of the edges that should appear 'selected' on the subdivision preview, taking into
+				// account all currently selected edges and polygons that share those edges
+				static TSet<FEdgeID> HighlightedEdgeIDs;
+				HighlightedEdgeIDs.Reset();
+
+				// @todo mesheditor perf: O(N) process to setup list of highlighted edges.  We could instead maintain this as selection changes.
+				for( FMeshElement& SelectedMeshElement : SelectedMeshElements )
+				{
+					if( SelectedMeshElement.IsValidMeshElement() &&
+						SelectedMeshElement.Component == Component &&
+						SelectedMeshElement.ElementAddress.SubMeshAddress == EditableMesh->GetSubMeshAddress() )
+					{
+						if( SelectedMeshElement.ElementAddress.ElementType == EEditableMeshElementType::Edge )
+						{
+							// OK, this edge is selected so it's definitely going to be highlighted
+							HighlightedEdgeIDs.Add( FEdgeID( SelectedMeshElement.ElementAddress.ElementID ) );
+						}
+						else if( SelectedMeshElement.ElementAddress.ElementType == EEditableMeshElementType::Polygon )
+						{
+							// Find all of this polygons edges.  They'll all need to be highlighted.
+							static TArray<FEdgeID> SelectedPolygonPerimeterEdgeIDs;
+							EditableMesh->GetPolygonPerimeterEdges(
+								FPolygonRef( SelectedMeshElement.ElementAddress.SectionID, FPolygonID( SelectedMeshElement.ElementAddress.ElementID ) ),
+								/* Out */ SelectedPolygonPerimeterEdgeIDs );
+
+							HighlightedEdgeIDs.Append( SelectedPolygonPerimeterEdgeIDs );
+						}
+					}
+				}
+
+
 				const float SizeBias = 0.0f;
 				const FColor SubdivisionEdgeColorAndOpacity = FLinearColor( 0.05f, 0.05f, 0.05f, 0.6f ).ToFColor( false );
 				const FColor BaseCageCounterpartEdgeColorAndOpacity = FLinearColor( 0.0f, 0.0f, 0.2f, 0.8f ).ToFColor( false );
+				const FColor SelectedBaseCageCounterpartEdgeColorAndOpacity = FLinearColor::White.CopyWithNewOpacity( 0.8f ).ToFColor( false );
 
 				const FSubdivisionLimitData& SubdivisionLimitData = EditableMesh->GetSubdivisionLimitData();
 
@@ -2270,7 +2302,18 @@ void FMeshEditorMode::Render( const FSceneView* SceneView, FViewport* Viewport,
 						WorldSpaceEdgeVertexPositions[ EdgeVertexNumber ] = ComponentToWorldMatrix.TransformPosition( EdgeVertexPositions[ EdgeVertexNumber ] );
 					}
 
-					const FColor ColorAndOpacity = SubdividedWireEdge.bIsBaseCageCounterpartEdge ? BaseCageCounterpartEdgeColorAndOpacity : SubdivisionEdgeColorAndOpacity;
+					FColor ColorAndOpacity = SubdivisionEdgeColorAndOpacity;
+					if( SubdividedWireEdge.CounterpartEdgeID != FEdgeID::Invalid )
+					{
+						if( HighlightedEdgeIDs.Contains( SubdividedWireEdge.CounterpartEdgeID ) )
+						{
+							ColorAndOpacity = SelectedBaseCageCounterpartEdgeColorAndOpacity;
+						}
+						else
+						{
+							ColorAndOpacity = BaseCageCounterpartEdgeColorAndOpacity;
+						}
+					}					
 
 					const bool bApplyDepthBias = true;
 					AddThickLineToDynamicMesh( CameraToWorld, WorldSpaceEdgeVertexPositions, ColorAndOpacity, SizeBias, bApplyDepthBias, VertexAndEdgeMeshBuilder );
@@ -2354,7 +2397,7 @@ void FMeshEditorMode::Render( const FSceneView* SceneView, FViewport* Viewport,
 	}
 
 
-	// Draw hovered and selected elements
+	// Draw hovered elements
 	{
 		// Only draw hover if we're not in the middle of an interactive edit
 		if( ActiveAction == NAME_None )
diff --git a/Engine/Source/Runtime/MeshEditingRuntime/EditableMesh.cpp b/Engine/Source/Runtime/MeshEditingRuntime/EditableMesh.cpp
--- a/Engine/Source/Runtime/MeshEditingRuntime/EditableMesh.cpp
+++ b/Engine/Source/Runtime/MeshEditingRuntime/EditableMesh.cpp
@@ -264,6 +264,31 @@ void UEditableMesh::GetEdgeConnectedPolygons( const FEdgeID EdgeID, TArray<FPoly
 }
 
 
+FEdgeID UEditableMesh::GetEdgeThatConnectsVertices( const FVertexID VertexID0, const FVertexID VertexID1 ) const
+{
+	FEdgeID EdgeID = FEdgeID::Invalid;
+
+	const int32 Vertex0ConnectedEdgeCount = GetVertexConnectedEdgeCount( VertexID0 );
+	for( int32 Vertex0ConnectedEdgeNumber = 0; Vertex0ConnectedEdgeNumber < Vertex0ConnectedEdgeCount; ++Vertex0ConnectedEdgeNumber )
+	{
+		const FEdgeID Vertex0ConnectedEdge = GetVertexConnectedEdge( VertexID0, Vertex0ConnectedEdgeNumber );
+
+		FVertexID EdgeVertexID0, EdgeVertexID1;
+		GetEdgeVertices( Vertex0ConnectedEdge, /* Out */ EdgeVertexID0, /* Out */ EdgeVertexID1 );
+
+		if( ( EdgeVertexID0 == VertexID0 && EdgeVertexID1 == VertexID1 ) ||
+			( EdgeVertexID0 == VertexID1 && EdgeVertexID1 == VertexID0 ) )
+		{
+			EdgeID = Vertex0ConnectedEdge;
+			break;
+		}
+	}
+
+	check( EdgeID != FEdgeID::Invalid );
+	return EdgeID;
+}
+
+
 void UEditableMesh::GetEdgeLoopElements( const FEdgeID EdgeID, TArray<FEdgeID>& EdgeLoopIDs ) const
 {
 	EdgeLoopIDs.Reset();
@@ -1462,6 +1487,9 @@ void UEditableMesh::GenerateOpenSubdivLimitSurfaceData()
 			genTangSpaceDefault( &MikkTContext );
 		}
 
+
+		// Generate our edge information for the subdivided mesh.  We'll also figure out which subdivided edges have a 
+		// counterpart on the base cage mesh, so tools can display this information to the user.
 		{
 			const int32 LimitEdgeCount = OsdLimitLevel.GetNumEdges();
 			for( int32 LimitEdgeNumber = 0; LimitEdgeNumber < LimitEdgeCount; ++LimitEdgeNumber )
@@ -1476,66 +1504,74 @@ void UEditableMesh::GenerateOpenSubdivLimitSurfaceData()
 				SubdividedWireEdge.EdgeVertex1PositionIndex = OsdLimitEdgeVertices[ 1 ];
 
 				// Default to not highlighting this edge as a base cage counterpart.  We'll actually figure this out below.
-				SubdividedWireEdge.bIsBaseCageCounterpartEdge = false;
+				SubdividedWireEdge.CounterpartEdgeID = FEdgeID::Invalid;
 			}
 
-
-
 			{
-				// Figure out all of the unique polygon edges that make up the base cage
-				static TArray<int32> BaseCageEdges;
-				BaseCageEdges.Reset();
+				static TSet<int32> BaseCageEdgeSet;
+				BaseCageEdgeSet.Reset();
+
+				const OpenSubdiv::Far::TopologyLevel& OsdBaseCageLevel = OsdTopologyRefiner->GetLevel( 0 );
+				const int32 BaseCageFaceCount = OsdBaseCageLevel.GetNumFaces();
+				for( int32 BaseCageFaceNumber = 0; BaseCageFaceNumber < BaseCageFaceCount; ++BaseCageFaceNumber )
 				{
-					const OpenSubdiv::Far::TopologyLevel& OsdBaseCageLevel = OsdTopologyRefiner->GetLevel( 0 );
-					const int32 BaseCageFaceCount = OsdBaseCageLevel.GetNumFaces();
-					for( int32 BaseCageFaceNumber = 0; BaseCageFaceNumber < BaseCageFaceCount; ++BaseCageFaceNumber )
+					const OpenSubdiv::Far::ConstIndexArray& OsdBaseCageFaceEdges = OsdBaseCageLevel.GetFaceEdges( BaseCageFaceNumber );
+					for( int32 FaceEdgeNumber = 0; FaceEdgeNumber < OsdBaseCageFaceEdges.size(); ++FaceEdgeNumber )
 					{
-						const OpenSubdiv::Far::ConstIndexArray& OsdBaseCageFaceEdges = OsdBaseCageLevel.GetFaceEdges( BaseCageFaceNumber );
-						for( int32 FaceEdgeNumber = 0; FaceEdgeNumber < OsdBaseCageFaceEdges.size(); ++FaceEdgeNumber )
+						const int32 BaseCageEdgeIndex = OsdBaseCageFaceEdges[ FaceEdgeNumber ];
+						bool bIsAlreadyInSet = false;
+						BaseCageEdgeSet.Add( BaseCageEdgeIndex, /* Out */ &bIsAlreadyInSet );
+						if( !bIsAlreadyInSet )
 						{
-							BaseCageEdges.AddUnique( OsdBaseCageFaceEdges[ FaceEdgeNumber ] );	// @todo mesheditor now: Any better way to get all of the edges in base cage and avoid AddUnique 
-						}
-					}
-				}
+							// Find our original edge ID for each of the OpenSubdiv base cage edges
+							const OpenSubdiv::Far::ConstIndexArray& OsdEdgeVertices = OsdBaseCageLevel.GetEdgeVertices( BaseCageEdgeIndex );
+							check( OsdEdgeVertices.size() == 2 ); // Edges always connect two vertices
+																	// Figure out which edge goes with these vertices
+							const FEdgeID BaseCageEdgeID = GetEdgeThatConnectsVertices( FVertexID( OsdEdgeVertices[ 0 ] ), FVertexID( OsdEdgeVertices[ 1 ] ) );
+
+							// Go through and determine the limit child edges of all of the original base cage edges by drilling down through
+							// the subdivision hierarchy
+							int32 NextScratchBufferIndex = 0;
+							static TArray<int32> ScratchChildEdges[ 2 ];
+							ScratchChildEdges[ 0 ].Reset();
+							ScratchChildEdges[ 1 ].Reset();
+
+							// Fill in our source buffer with the starting edge
+							ScratchChildEdges[ NextScratchBufferIndex ].Add( BaseCageEdgeIndex );
+							NextScratchBufferIndex = !NextScratchBufferIndex;
+
+							for( int32 RefinementLevel = 0; RefinementLevel < SubdivisionCount; ++RefinementLevel )
+							{
+								const OpenSubdiv::Far::TopologyLevel& OsdLevel = OsdTopologyRefiner->GetLevel( RefinementLevel );
 
-				// Go through and determine the limit child edges of all of the original base cage edges by drilling down through
-				// the subdivision heirarchy
-				int32 NextScratchBufferIndex = 0;
-				static TArray<int32> ScratchChildEdges[ 2 ];
-				ScratchChildEdges[ 0 ].Reset();
-				ScratchChildEdges[ 1 ].Reset();
-				{
-					for( int32 RefinementLevel = 0; RefinementLevel < SubdivisionCount; ++RefinementLevel )
-					{
-						const OpenSubdiv::Far::TopologyLevel& OsdLevel = OsdTopologyRefiner->GetLevel( RefinementLevel );
+								const TArray<int32>& SourceChildEdges = ScratchChildEdges[ !NextScratchBufferIndex ];
+
+								TArray<int32>& DestChildEdges = ScratchChildEdges[ NextScratchBufferIndex ];
+								DestChildEdges.Reset();
 
-						const TArray<int32>& SourceChildEdges = ( RefinementLevel == 0 ) ? BaseCageEdges : ScratchChildEdges[ !NextScratchBufferIndex ];
+								for( int32 SourceEdgeNumber = 0; SourceEdgeNumber < SourceChildEdges.Num(); ++SourceEdgeNumber )
+								{
+									const OpenSubdiv::Far::ConstIndexArray& OsdChildEdges = OsdLevel.GetEdgeChildEdges( SourceChildEdges[ SourceEdgeNumber ] );
+									for( int32 ChildEdgeNumber = 0; ChildEdgeNumber < OsdChildEdges.size(); ++ChildEdgeNumber )
+									{
+										DestChildEdges.Add( OsdChildEdges[ ChildEdgeNumber ] );
+									}
+								}
 
-						TArray<int32>& DestChildEdges = ScratchChildEdges[ NextScratchBufferIndex ];
-						DestChildEdges.Reset();
+								NextScratchBufferIndex = !NextScratchBufferIndex;
+							}
 
-						for( int32 SourceEdgeNumber = 0; SourceEdgeNumber < SourceChildEdges.Num(); ++SourceEdgeNumber )
-						{
-							const OpenSubdiv::Far::ConstIndexArray& OsdChildEdges = OsdLevel.GetEdgeChildEdges( SourceChildEdges[ SourceEdgeNumber ] );
-							for( int32 ChildEdgeNumber = 0; ChildEdgeNumber < OsdChildEdges.size(); ++ChildEdgeNumber )
+							// No go back and update our subdivided wire edges, marking the edges that we determined were descendants of the base cage edges
+							const TArray<int32>& BaseCageCounterpartEdgesAtLimit = ScratchChildEdges[ !NextScratchBufferIndex ];
+							for( const int32 CounterpartEdgeAtLimit : BaseCageCounterpartEdgesAtLimit )
 							{
-								DestChildEdges.Add( OsdChildEdges[ ChildEdgeNumber ] );
+								check( CounterpartEdgeAtLimit < SubdivisionLimitData.SubdividedWireEdges.Num() );
+								SubdivisionLimitData.SubdividedWireEdges[ CounterpartEdgeAtLimit ].CounterpartEdgeID = BaseCageEdgeID;
 							}
 						}
-
-						NextScratchBufferIndex = !NextScratchBufferIndex;
 					}
 				}
-
-				// No go back and update our subdivided wire edges, marking the edges that we determined were decendants of the base cage edges
-				const TArray<int32>& BaseCageCounterpartEdgesAtLimit = ScratchChildEdges[ !NextScratchBufferIndex ];
-				for( const int32 CounterpartEdgeAtLimit : BaseCageCounterpartEdgesAtLimit )
-				{
-					check( CounterpartEdgeAtLimit < SubdivisionLimitData.SubdividedWireEdges.Num() );
-					SubdivisionLimitData.SubdividedWireEdges[ CounterpartEdgeAtLimit ].bIsBaseCageCounterpartEdge = true;
-				}
 			}
-
 		}
 	}
 }
diff --git a/Engine/Source/Runtime/MeshEditingRuntime/Public/EditableMesh.h b/Engine/Source/Runtime/MeshEditingRuntime/Public/EditableMesh.h
--- a/Engine/Source/Runtime/MeshEditingRuntime/Public/EditableMesh.h
+++ b/Engine/Source/Runtime/MeshEditingRuntime/Public/EditableMesh.h
@@ -239,6 +239,7 @@ public:
 	UFUNCTION( BlueprintPure, Category="Editable Mesh" ) void GetEdgeVertices( const FEdgeID EdgeID, FVertexID& OutEdgeVertexID0, FVertexID& OutEdgeVertexID1 ) const;
 	UFUNCTION( BlueprintPure, Category="Editable Mesh" ) void GetEdgeConnectedPolygons( const FEdgeID EdgeID, TArray<FPolygonRef>& OutConnectedPolygonRefs ) const;
 	UFUNCTION( BlueprintPure, Category="Editable Mesh" ) void GetEdgeLoopElements( const FEdgeID EdgeID, TArray<FEdgeID>& EdgeLoopIDs ) const;
+	UFUNCTION( BlueprintPure, Category="Editable Mesh" ) FEdgeID GetEdgeThatConnectsVertices( const FVertexID VertexID0, const FVertexID VertexID1 ) const;
 	UFUNCTION( BlueprintPure, Category="Editable Mesh" ) int32 GetPolygonPerimeterEdgeCount( const FPolygonRef PolygonRef ) const;
 	UFUNCTION( BlueprintPure, Category="Editable Mesh" ) int32 GetPolygonHoleEdgeCount( const FPolygonRef PolygonRef, const int32 HoleNumber ) const;
 	UFUNCTION( BlueprintPure, Category="Editable Mesh" ) void GetPolygonPerimeterVertices( const FPolygonRef PolygonRef, TArray<FVertexID>& OutPolygonPerimeterVertexIDs ) const;
diff --git a/Engine/Source/Runtime/MeshEditingRuntime/Public/EditableMeshTypes.h b/Engine/Source/Runtime/MeshEditingRuntime/Public/EditableMeshTypes.h
--- a/Engine/Source/Runtime/MeshEditingRuntime/Public/EditableMeshTypes.h
+++ b/Engine/Source/Runtime/MeshEditingRuntime/Public/EditableMeshTypes.h
@@ -625,7 +625,7 @@ struct FSubdividedWireEdge
 
 	/** True if this edge is a counterpart to an original base cage edge of the mesh.  Otherwise it's a new edge that exists only
 	    in the subdivision surfaces */
-	bool bIsBaseCageCounterpartEdge;
+	FEdgeID CounterpartEdgeID;
 };
 
 
