Mesh Editor: Extrude multiple polygons together!
- When multiple polygons are selected that are neighbors, Extrude will now keep those polygons connected by default
- Extrusion distances for polygons are now influenced by their vertices' intersection points with neighbor polygons.  This makes for much more sensible extrusions when extruding multiple polygons together.
- UEditableMesh::ExtrudePolygons() has a new parameter 'bKeepNeighborsTogether' to enable this feature (always 'true' in Mesh Editor for now.  We'll make it an option, eventually.)

diff --git a/Engine/Plugins/Editor/MeshEditor/Source/PolygonModeling/ExtrudePolygon.cpp b/Engine/Plugins/Editor/MeshEditor/Source/PolygonModeling/ExtrudePolygon.cpp
--- a/Engine/Plugins/Editor/MeshEditor/Source/PolygonModeling/ExtrudePolygon.cpp
+++ b/Engine/Plugins/Editor/MeshEditor/Source/PolygonModeling/ExtrudePolygon.cpp
@@ -169,7 +169,8 @@ void UExtrudePolygonCommand::ApplyDuringDrag( IMeshEditorModeEditingContract& Me
 				// Create a copy of the polygon with new extruded polygons for each edge
 				static TArray<FPolygonRef> NewExtrudedFrontPolygons;
 				NewExtrudedFrontPolygons.Reset();
-				EditableMesh->ExtrudePolygons( PolygonRefsToExtrude, ComponentSpaceExtrudeDistance, /* Out */ NewExtrudedFrontPolygons );
+				const bool bKeepNeighborsTogether = true;	// @todo mesheditor: Make configurable in UI
+				EditableMesh->ExtrudePolygons( PolygonRefsToExtrude, ComponentSpaceExtrudeDistance, bKeepNeighborsTogether, /* Out */ NewExtrudedFrontPolygons );
 
 
 				// Make sure the new polygons are selected.  The old polygon was deleted and will become deselected automatically.
diff --git a/Engine/Source/Runtime/MeshEditingRuntime/EditableMesh.cpp b/Engine/Source/Runtime/MeshEditingRuntime/EditableMesh.cpp
--- a/Engine/Source/Runtime/MeshEditingRuntime/EditableMesh.cpp
+++ b/Engine/Source/Runtime/MeshEditingRuntime/EditableMesh.cpp
@@ -3735,14 +3735,15 @@ void UEditableMesh::TryToRemoveVertex( const FVertexID VertexID, bool& bOutWasVe
 }
 
 
-void UEditableMesh::ExtrudePolygons( const TArray<FPolygonRef>& Polygons, const float ExtrudeDistance, TArray<FPolygonRef>& OutNewExtrudedFrontPolygons )
+void UEditableMesh::ExtrudePolygons( const TArray<FPolygonRef>& Polygons, const float ExtrudeDistance, const bool bKeepNeighborsTogether, TArray<FPolygonRef>& OutNewExtrudedFrontPolygons )
 {
-	// @todo mesheditor: For extrusion, we really need a mode where selected polygons that are connected will all
-	// be extruded together (no new edges created between the neighbors.)  See other modeling programs for examples.
-	// --> This should also work with multiple selected edges and vertices too, ideally
-
 	// @todo mesheditor perf: We can make this much faster by batching up polygons together.  Just be careful about how neighbors are handled.
 	OutNewExtrudedFrontPolygons.Reset();
+
+	// Convert our incoming polygon array to a TSet so we can lookup quickly to see which polygons in the mesh sare members of the set
+	static TSet<FPolygonRef> PolygonsSet;
+	PolygonsSet.Reset();
+	PolygonsSet.Append( Polygons );
 	
 	static TArray<FPolygonRef> AllNewPolygons;
 	AllNewPolygons.Reset();
@@ -3756,182 +3757,302 @@ void UEditableMesh::ExtrudePolygons( const TArray<FPolygonRef>& Polygons, const
 	static TArray<FVertexAttributesForPolygon> VertexAttributesForPolygons;
 	VertexAttributesForPolygons.Reset();
 
+	// First, let's figure out which of the polygons we were asked to extrude share edges or vertices.  We'll keep those
+	// edges intact!
+	static TMap<FEdgeID, uint32> EdgeUsageCounts;	// Maps an edge ID to the number of times it is referenced by the incoming polygons
+	EdgeUsageCounts.Reset();
+	static TSet<FVertexID> UniqueVertexIDs;
+	UniqueVertexIDs.Reset();
 	for( int32 PolygonIter = 0; PolygonIter < Polygons.Num(); ++PolygonIter )
 	{
 		const FPolygonRef& PolygonRef = Polygons[ PolygonIter ];
 
-		// We'll need this polygon's normal to figure out where to put the extruded copy of the polygon
-		const FVector PolygonNormal = ComputePolygonNormal( PolygonRef );
+		static TArray<FEdgeID> PolygonPerimeterEdgeIDs;
+		GetPolygonPerimeterEdges( PolygonRef, /* Out */ PolygonPerimeterEdgeIDs );
 
-		// Create a new set of vertices for the extruded polygon
-		static TArray<FVertexID> ExtrudedVertexIDs;
-		ExtrudedVertexIDs.Reset();
+		for( const FEdgeID EdgeID : PolygonPerimeterEdgeIDs )
 		{
-			const int32 PerimeterVertexCount = this->GetPolygonPerimeterVertexCount( PolygonRef );
-			ExtrudedVertexIDs.SetNumUninitialized( PerimeterVertexCount, false );
+			FVertexID EdgeVertexIDs[ 2 ];
+			GetEdgeVertices( EdgeID, /* Out */ EdgeVertexIDs[ 0 ], /* Out */ EdgeVertexIDs[ 1 ] );
 
-			// Create some vertices
-			CreateEmptyVertexRange( PerimeterVertexCount, /* Out */ ExtrudedVertexIDs );
+			uint32* EdgeUsageCountPtr = EdgeUsageCounts.Find( EdgeID );
+			if( EdgeUsageCountPtr == nullptr )
+			{
+				EdgeUsageCounts.Add( EdgeID, 1 );
+			}
+			else
+			{
+				++( *EdgeUsageCountPtr );
+			}
 		}
-		int32 NextAvailableExtrudedVertexIDNumber = 0;
 
-		static TArray<FVertexID> PolygonVertexIDs;
-		PolygonVertexIDs.Reset();
-		this->GetPolygonPerimeterVertices( PolygonRef, /* Out */ PolygonVertexIDs );
 
-		// Map all of the edge vertices to their new extruded counterpart
-		static TMap< FVertexID, FVertexID> VertexIDToExtrudedCopyMap;
-		VertexIDToExtrudedCopyMap.Reset();
+		static TArray<FVertexID> PolygonPerimeterVertexIDs;
+		this->GetPolygonPerimeterVertices( PolygonRef, /* Out */ PolygonPerimeterVertexIDs );
+
+		for( const FVertexID VertexID : PolygonPerimeterVertexIDs )
 		{
-			const int32 PerimeterEdgeCount = this->GetPolygonPerimeterEdgeCount( PolygonRef );
-			for( int32 PerimeterEdgeNumber = 0; PerimeterEdgeNumber < PerimeterEdgeCount; ++PerimeterEdgeNumber )
-			{
-				FVertexID EdgeVertexIDs[ 2 ];
-				{
-					// @todo mesheditor perf: O(log N^2) iteration here. For every edge, for every edge up to this index.  Need to clean this up. 
-					//		--> Also, there are quite a few places where we are stepping through edges in perimeter-order.  We need to have a nice way to walk that.
-					bool bEdgeWindingIsReversedForPolygon;
-					const FEdgeID EdgeID = this->GetPolygonPerimeterEdge( PolygonRef, PerimeterEdgeNumber, /* Out */ bEdgeWindingIsReversedForPolygon );
+			UniqueVertexIDs.Add( VertexID );
+		}
+	}
 
-					// After extruding, all of the edges of the original polygon become hard edges
-					FAttributesForEdge& AttributesForEdge = *new( AttributesForEdges ) FAttributesForEdge();
-					AttributesForEdge.EdgeID = EdgeID;
-					AttributesForEdge.EdgeAttributes.Attributes.Add( FMeshElementAttributeData( UEditableMeshAttribute::EdgeIsHard(), 0, FVector4( 1.0f ) ) );
-					AttributesForEdge.EdgeAttributes.Attributes.Add( FMeshElementAttributeData( UEditableMeshAttribute::EdgeCreaseSharpness(), 0, FVector4( 1.0f ) ) );
+	const int32 NumVerticesToCreate = UniqueVertexIDs.Num();
+
+	// Create new vertices for all of the extruded polygons
+	static TArray<FVertexID> ExtrudedVertexIDs;
+	ExtrudedVertexIDs.Reset();
+	ExtrudedVertexIDs.Reserve( NumVerticesToCreate );
+	CreateEmptyVertexRange( NumVerticesToCreate, /* Out */ ExtrudedVertexIDs );
+	int32 NextAvailableExtrudedVertexIDNumber = 0;
+
+
+	static TMap<FVertexID, FVertexID> VertexIDToExtrudedCopy;
+	VertexIDToExtrudedCopy.Reset();
+
+	for( int32 PassIndex = 0; PassIndex < 2; ++PassIndex )
+	{
+		// Extrude all of the shared edges first, then do the non-shared edges.  This is to make sure that a vertex doesn't get offset
+		// without taking into account all of the connected polygons in our set
+		const bool bIsExtrudingSharedEdges = ( PassIndex == 0 );
 
-					this->GetEdgeVertices( EdgeID, /* Out */ EdgeVertexIDs[0], /* Out */ EdgeVertexIDs[1] );
+		for( int32 PolygonIter = 0; PolygonIter < Polygons.Num(); ++PolygonIter )
+		{
+			const FPolygonRef& PolygonRef = Polygons[ PolygonIter ];
+
+			if( !bKeepNeighborsTogether )
+			{
+				VertexIDToExtrudedCopy.Reset();
+			}
+
+			// Map all of the edge vertices to their new extruded counterpart
+			const int32 PerimeterEdgeCount = GetPolygonPerimeterEdgeCount( PolygonRef );
+			for( int32 PerimeterEdgeNumber = 0; PerimeterEdgeNumber < PerimeterEdgeCount; ++PerimeterEdgeNumber )
+			{
+				// @todo mesheditor perf: We can change GetPolygonPerimeterEdges() to have a version that returns whether winding is reversed or not, and avoid this call entirely.
+				// @todo mesheditor perf: O(log N^2) iteration here. For every edge, for every edge up to this index.  Need to clean this up. 
+				//		--> Also, there are quite a few places where we are stepping through edges in perimeter-order.  We need to have a nice way to walk that.
+				bool bEdgeWindingIsReversedForPolygon;
+				const FEdgeID EdgeID = this->GetPolygonPerimeterEdge( PolygonRef, PerimeterEdgeNumber, /* Out */ bEdgeWindingIsReversedForPolygon );
 
+				const bool bIsSharedEdge = bKeepNeighborsTogether && EdgeUsageCounts[ EdgeID ] > 1;
+				if( bIsSharedEdge == bIsExtrudingSharedEdges )
+				{
+					FVertexID EdgeVertexIDs[ 2 ];
+					GetEdgeVertices( EdgeID, /* Out */ EdgeVertexIDs[ 0 ], /* Out */ EdgeVertexIDs[ 1 ] );
 					if( bEdgeWindingIsReversedForPolygon )
 					{
-						const FVertexID TempVertexID = EdgeVertexIDs[ 1 ];
-						EdgeVertexIDs[ 1 ] = EdgeVertexIDs[ 0 ];
-						EdgeVertexIDs[ 0 ] = TempVertexID;
+						::Swap( EdgeVertexIDs[ 0 ], EdgeVertexIDs[ 1 ] );
 					}
-				}
 
-				FVertexID ExtrudedEdgeVertexIDs[ 2 ];
-				for( int32 EdgeVertexNumber = 0; EdgeVertexNumber < 2; ++EdgeVertexNumber )
-				{
-					const FVertexID EdgeVertexID = EdgeVertexIDs[ EdgeVertexNumber ];
-					FVertexID* ExtrudedEdgeVertexIDPtr = VertexIDToExtrudedCopyMap.Find( EdgeVertexID );
-					if( ExtrudedEdgeVertexIDPtr != nullptr )
+					if( !bIsSharedEdge )
 					{
-						ExtrudedEdgeVertexIDs[ EdgeVertexNumber ] = *ExtrudedEdgeVertexIDPtr;
+						// After extruding, all of the edges of the original polygon become hard edges
+						FAttributesForEdge& AttributesForEdge = *new( AttributesForEdges ) FAttributesForEdge();
+						AttributesForEdge.EdgeID = EdgeID;
+						AttributesForEdge.EdgeAttributes.Attributes.Add( FMeshElementAttributeData( UEditableMeshAttribute::EdgeIsHard(), 0, FVector4( 1.0f ) ) );
+						AttributesForEdge.EdgeAttributes.Attributes.Add( FMeshElementAttributeData( UEditableMeshAttribute::EdgeCreaseSharpness(), 0, FVector4( 1.0f ) ) );
 					}
-					else
+
+					FVertexID ExtrudedEdgeVertexIDs[ 2 ];
+					for( int32 EdgeVertexNumber = 0; EdgeVertexNumber < 2; ++EdgeVertexNumber )
 					{
-						// Create a copy of this vertex for the extruded face
-						const FVertexID ExtrudedVertexID = ExtrudedVertexIDs[ NextAvailableExtrudedVertexIDNumber++ ];
+						const FVertexID EdgeVertexID = EdgeVertexIDs[ EdgeVertexNumber ];
+						FVertexID* ExtrudedEdgeVertexIDPtr = VertexIDToExtrudedCopy.Find( EdgeVertexID );
+
+						// @todo mesheditor extrude: Ideally we would detect whether the vertex that was already extruded came from a edge
+						// from a polygon that does not actually share an edge with any polygons this polygon shares an edge with.  This
+						// would avoid the problem where extruding two polygons that are connected only by a vertex are not extruded
+						// separately.
+						const bool bVertexIsSharedByAnEdgeOfAnotherSelectedPolygon = false;
+						if( ExtrudedEdgeVertexIDPtr != nullptr && !bVertexIsSharedByAnEdgeOfAnotherSelectedPolygon )
+						{
+							ExtrudedEdgeVertexIDs[ EdgeVertexNumber ] = *ExtrudedEdgeVertexIDPtr;
+						}
+						else
+						{
+							// Create a copy of this vertex for the extruded face
+							const FVertexID ExtrudedVertexID = ExtrudedVertexIDs[ NextAvailableExtrudedVertexIDNumber++ ];
 
-						ExtrudedEdgeVertexIDPtr = &VertexIDToExtrudedCopyMap.Add( EdgeVertexID, ExtrudedVertexID );
+							ExtrudedEdgeVertexIDPtr = &VertexIDToExtrudedCopy.Add( EdgeVertexID, ExtrudedVertexID );
 
-						// Push the vertex out along the polygon's normal
-						const FVector OriginalVertexPosition = this->GetVertexAttribute( EdgeVertexID, UEditableMeshAttribute::VertexPosition(), 0 );
+							// Push the vertex out along the polygon's normal
+							const FVector OriginalVertexPosition = this->GetVertexAttribute( EdgeVertexID, UEditableMeshAttribute::VertexPosition(), 0 );
 
-						const FVector ExtrudedVertexPosition = OriginalVertexPosition + ExtrudeDistance * PolygonNormal;
+							FVector ExtrudedVertexPosition;
+							if( bIsSharedEdge )
+							{
+								// Get all of the polygons that share this edge that were part of the set of polygons passed in.  We'll
+								// generate an extrude direction that's the average of those polygon normals.
+								FVector ExtrudeDirection = FVector::ZeroVector;
 
-						// Fill in the vertex
-						FAttributesForVertex& AttributesForVertex = *new( AttributesForVertices ) FAttributesForVertex();
-						AttributesForVertex.VertexID = ExtrudedVertexID;
-						AttributesForVertex.VertexAttributes.Attributes.Add( FMeshElementAttributeData(
-							UEditableMeshAttribute::VertexPosition(), 
-							0, 
-							ExtrudedVertexPosition ) );
-					}
-					ExtrudedEdgeVertexIDs[ EdgeVertexNumber ] = *ExtrudedEdgeVertexIDPtr;
-				}
+								static TArray<FPolygonRef> ConnectedPolygonRefs;
+								GetVertexConnectedPolygons( EdgeVertexID, /* Out */ ConnectedPolygonRefs );
 
-				// Add a face that connects the edge to it's extruded counterpart edge
-				{
-					FVertexID OriginalVertexIDs[ 4 ];
-
-					static TArray<FVertexAndAttributes> NewSidePolygonVertices;
-					NewSidePolygonVertices.Reset( 4 );
-					NewSidePolygonVertices.SetNum( 4, false );	// Always four edges in an extruded face
-
-					NewSidePolygonVertices[ 0 ].VertexID = EdgeVertexIDs[ 1 ];
-					OriginalVertexIDs[ 0 ] = EdgeVertexIDs[ 1 ];
-					NewSidePolygonVertices[ 1 ].VertexID = EdgeVertexIDs[ 0 ];
-					OriginalVertexIDs[ 1 ] = EdgeVertexIDs[ 0 ];
-					NewSidePolygonVertices[ 2 ].VertexID = ExtrudedEdgeVertexIDs[ 0 ];
-					OriginalVertexIDs[ 2 ] = EdgeVertexIDs[ 0 ];
-					NewSidePolygonVertices[ 3 ].VertexID = ExtrudedEdgeVertexIDs[ 1 ];
-					OriginalVertexIDs[ 3 ] = EdgeVertexIDs[ 1 ];
-
-					FPolygonRef NewSidePolygonRef;	// Filled in below
-					{
-						static TArray<FPolygonToCreate> PolygonsToCreate;
-						PolygonsToCreate.Reset();
+								static TArray<FPolygonRef> NeighborPolygonRefs;
+								NeighborPolygonRefs.Reset();
+								for( const FPolygonRef ConnectedPolygonRef : ConnectedPolygonRefs )
+								{
+									// We only care about polygons that are members of the set of polygons we were asked to extrude
+									if( PolygonsSet.Contains( ConnectedPolygonRef ) )
+									{
+										NeighborPolygonRefs.Add( ConnectedPolygonRef );
 
-						// Create the polygon
-						// @todo mesheditor perf: Ideally we support creating multiple polygons at once and batching up the work
-						FPolygonToCreate& PolygonToCreate = *new( PolygonsToCreate ) FPolygonToCreate();
-						PolygonToCreate.SectionID = PolygonRef.SectionID;
-						PolygonToCreate.PerimeterVertices = NewSidePolygonVertices;	// @todo mesheditor perf: Copying static array here, ideally allocations could be avoided
+										// We'll need this polygon's normal to figure out where to put the extruded copy of the polygon
+										const FVector NeighborPolygonNormal = ComputePolygonNormal( ConnectedPolygonRef );
+										ExtrudeDirection += NeighborPolygonNormal;
+									}
+								}
+								ExtrudeDirection.Normalize();
+
+
+								// OK, we have the direction to extrude for this vertex.  Now we need to know how far to extrude.  We'll
+								// loop over all of the neighbor polygons to this vertex, and choose the closest intersection point with our
+								// vertex's extrude direction and the neighbor polygon's extruded plane
+								FVector ClosestIntersectionPointWithExtrudedPlanes;
+								float ClosestIntersectionDistanceSquared = TNumericLimits<float>::Max();
 
-						// NOTE: We never create holes in side polygons
+								for( const FPolygonRef NeighborPolygonRef : NeighborPolygonRefs )
+								{
+									const FPlane NeighborPolygonPlane = ComputePolygonPlane( NeighborPolygonRef );
+
+									// Push the plane out
+									const FPlane ExtrudedPlane = [NeighborPolygonPlane, ExtrudeDistance]
+										{ 
+											FPlane NewPlane = NeighborPolygonPlane;
+											NewPlane.W += ExtrudeDistance; 
+											return NewPlane; 
+										}();
+
+									// Is this the closest intersection point so far?
+									const FVector IntersectionPointWithExtrudedPlane = FMath::RayPlaneIntersection( OriginalVertexPosition, ExtrudeDirection, ExtrudedPlane );
+									const float IntersectionDistanceSquared = FVector::DistSquared( OriginalVertexPosition, IntersectionPointWithExtrudedPlane );
+									if( IntersectionDistanceSquared < ClosestIntersectionDistanceSquared )
+									{
+										ClosestIntersectionPointWithExtrudedPlanes = IntersectionPointWithExtrudedPlane;
+										ClosestIntersectionDistanceSquared = IntersectionDistanceSquared;
+									}
+								}
 
-						static TArray<FPolygonRef> NewPolygonRefs;
-						NewPolygonRefs.Reset();
-						static TArray<FEdgeID> NewEdgeIDs;
-						NewEdgeIDs.Reset();
-						CreatePolygons( PolygonsToCreate, /* Out */ NewPolygonRefs, /* Out */ NewEdgeIDs );	// @todo mesheditor perf: Extra allocatons/copies: Ideally MoveTemp() here but we can't move a STATIC local!
+								ExtrudedVertexPosition = ClosestIntersectionPointWithExtrudedPlanes;
+							}
+							else
+							{
+								// We'll need this polygon's normal to figure out where to put the extruded copy of the polygon
+								const FVector PolygonNormal = ComputePolygonNormal( PolygonRef );
+								ExtrudedVertexPosition = OriginalVertexPosition + ExtrudeDistance * PolygonNormal;
+							}
 
-						NewSidePolygonRef = NewPolygonRefs[ 0 ];
+							// Fill in the vertex
+							FAttributesForVertex& AttributesForVertex = *new( AttributesForVertices ) FAttributesForVertex();
+							AttributesForVertex.VertexID = ExtrudedVertexID;
+							AttributesForVertex.VertexAttributes.Attributes.Add( FMeshElementAttributeData(
+								UEditableMeshAttribute::VertexPosition(),
+								0,
+								ExtrudedVertexPosition ) );
+						}
+						ExtrudedEdgeVertexIDs[ EdgeVertexNumber ] = *ExtrudedEdgeVertexIDPtr;
 					}
-					AllNewPolygons.Add( NewSidePolygonRef );
 
-					// Copy polygon UVs from the original polygon's vertices
+					if( !bIsSharedEdge )
 					{
-						FVertexAttributesForPolygon& PolygonNewAttributes = *new( VertexAttributesForPolygons ) FVertexAttributesForPolygon();
-						PolygonNewAttributes.PolygonRef = NewSidePolygonRef;
-						PolygonNewAttributes.PerimeterVertexAttributeLists.SetNum( NewSidePolygonVertices.Num(), false );
+						// Add a face that connects the edge to it's extruded counterpart edge
+						FVertexID OriginalVertexIDs[ 4 ];
+
+						static TArray<FVertexAndAttributes> NewSidePolygonVertices;
+						NewSidePolygonVertices.Reset( 4 );
+						NewSidePolygonVertices.SetNum( 4, false );	// Always four edges in an extruded face
+
+						NewSidePolygonVertices[ 0 ].VertexID = EdgeVertexIDs[ 1 ];
+						OriginalVertexIDs[ 0 ] = EdgeVertexIDs[ 1 ];
+						NewSidePolygonVertices[ 1 ].VertexID = EdgeVertexIDs[ 0 ];
+						OriginalVertexIDs[ 1 ] = EdgeVertexIDs[ 0 ];
+						NewSidePolygonVertices[ 2 ].VertexID = ExtrudedEdgeVertexIDs[ 0 ];
+						OriginalVertexIDs[ 2 ] = EdgeVertexIDs[ 0 ];
+						NewSidePolygonVertices[ 3 ].VertexID = ExtrudedEdgeVertexIDs[ 1 ];
+						OriginalVertexIDs[ 3 ] = EdgeVertexIDs[ 1 ];
+
+						FPolygonRef NewSidePolygonRef;	// Filled in below
+						{
+							static TArray<FPolygonToCreate> PolygonsToCreate;
+							PolygonsToCreate.Reset();
+
+							// Create the polygon
+							// @todo mesheditor perf: Ideally we support creating multiple polygons at once and batching up the work
+							FPolygonToCreate& PolygonToCreate = *new( PolygonsToCreate ) FPolygonToCreate();
+							PolygonToCreate.SectionID = PolygonRef.SectionID;
+							PolygonToCreate.PerimeterVertices = NewSidePolygonVertices;	// @todo mesheditor perf: Copying static array here, ideally allocations could be avoided
+
+							// NOTE: We never create holes in side polygons
+
+							static TArray<FPolygonRef> NewPolygonRefs;
+							NewPolygonRefs.Reset();
+							static TArray<FEdgeID> NewEdgeIDs;
+							NewEdgeIDs.Reset();
+							CreatePolygons( PolygonsToCreate, /* Out */ NewPolygonRefs, /* Out */ NewEdgeIDs );	// @todo mesheditor perf: Extra allocatons/copies: Ideally MoveTemp() here but we can't move a STATIC local!
+
+							NewSidePolygonRef = NewPolygonRefs[ 0 ];
+						}
+						AllNewPolygons.Add( NewSidePolygonRef );
 
-						for( int32 NewVertexIter = 0; NewVertexIter < NewSidePolygonVertices.Num(); ++NewVertexIter )
+						// Copy polygon UVs from the original polygon's vertices
 						{
-							const int32 PolygonVertexNumber = this->FindPolygonPerimeterVertexNumberForVertex( PolygonRef, OriginalVertexIDs[ NewVertexIter ] );
-							check( PolygonVertexNumber != INDEX_NONE );
+							FVertexAttributesForPolygon& PolygonNewAttributes = *new( VertexAttributesForPolygons ) FVertexAttributesForPolygon();
+							PolygonNewAttributes.PolygonRef = NewSidePolygonRef;
+							PolygonNewAttributes.PerimeterVertexAttributeLists.SetNum( NewSidePolygonVertices.Num(), false );
 
-							TArray<FMeshElementAttributeData>& PerimeterVertexNewAttributes = PolygonNewAttributes.PerimeterVertexAttributeLists[ NewVertexIter ].Attributes;
-							for( int32 TextureCoordinateIndex = 0; TextureCoordinateIndex < TextureCoordinateCount; ++TextureCoordinateIndex )
+							for( int32 NewVertexIter = 0; NewVertexIter < NewSidePolygonVertices.Num(); ++NewVertexIter )
 							{
-								const FVector4 TextureCoordinate = this->GetPolygonPerimeterVertexAttribute(
+								const int32 PolygonVertexNumber = this->FindPolygonPerimeterVertexNumberForVertex( PolygonRef, OriginalVertexIDs[ NewVertexIter ] );
+								check( PolygonVertexNumber != INDEX_NONE );
+
+								TArray<FMeshElementAttributeData>& PerimeterVertexNewAttributes = PolygonNewAttributes.PerimeterVertexAttributeLists[ NewVertexIter ].Attributes;
+								for( int32 TextureCoordinateIndex = 0; TextureCoordinateIndex < TextureCoordinateCount; ++TextureCoordinateIndex )
+								{
+									const FVector4 TextureCoordinate = this->GetPolygonPerimeterVertexAttribute(
+										PolygonRef,
+										PolygonVertexNumber,
+										UEditableMeshAttribute::VertexTextureCoordinate(),
+										TextureCoordinateIndex );
+
+									PerimeterVertexNewAttributes.Emplace( UEditableMeshAttribute::VertexTextureCoordinate(), TextureCoordinateIndex, TextureCoordinate );
+								}
+
+								const FVector4 VertexColor = this->GetPolygonPerimeterVertexAttribute(
 									PolygonRef,
 									PolygonVertexNumber,
-									UEditableMeshAttribute::VertexTextureCoordinate(),
-									TextureCoordinateIndex );
+									UEditableMeshAttribute::VertexColor(),
+									0 );
 
-								PerimeterVertexNewAttributes.Emplace( UEditableMeshAttribute::VertexTextureCoordinate(), TextureCoordinateIndex, TextureCoordinate );
+								PerimeterVertexNewAttributes.Emplace( UEditableMeshAttribute::VertexColor(), 0, VertexColor );
 							}
-
-							const FVector4 VertexColor = this->GetPolygonPerimeterVertexAttribute(
-								PolygonRef,
-								PolygonVertexNumber,
-								UEditableMeshAttribute::VertexColor(),
-								0 );
-
-							PerimeterVertexNewAttributes.Emplace( UEditableMeshAttribute::VertexColor(), 0, VertexColor );
 						}
-					}
 
-					// All of this edges of the new polygon will be hard
-					{
-						const int32 NewPerimeterEdgeCount = this->GetPolygonPerimeterEdgeCount( NewSidePolygonRef );
-						for( int32 NewPerimeterEdgeNumber = 0; NewPerimeterEdgeNumber < NewPerimeterEdgeCount; ++NewPerimeterEdgeNumber )
+						// All of this edges of the new polygon will be hard
 						{
-							bool bEdgeWindingIsReversedForPolygon;
-							const FEdgeID EdgeID = this->GetPolygonPerimeterEdge( NewSidePolygonRef, NewPerimeterEdgeNumber, /* Out */ bEdgeWindingIsReversedForPolygon );
+							const int32 NewPerimeterEdgeCount = this->GetPolygonPerimeterEdgeCount( NewSidePolygonRef );
+							for( int32 NewPerimeterEdgeNumber = 0; NewPerimeterEdgeNumber < NewPerimeterEdgeCount; ++NewPerimeterEdgeNumber )
+							{
+								bool bNewEdgeWindingIsReversedForPolygon;
+								const FEdgeID NewEdgeID = this->GetPolygonPerimeterEdge( NewSidePolygonRef, NewPerimeterEdgeNumber, /* Out */ bNewEdgeWindingIsReversedForPolygon );
 
-							// New side polygons get hard edges, but not hard crease weights
-							FAttributesForEdge& AttributesForEdge = *new( AttributesForEdges ) FAttributesForEdge();
-							AttributesForEdge.EdgeID = EdgeID;
-							AttributesForEdge.EdgeAttributes.Attributes.Add( FMeshElementAttributeData( UEditableMeshAttribute::EdgeIsHard(), 0, FVector4( 1.0f ) ) );
+								// New side polygons get hard edges, but not hard crease weights
+								FAttributesForEdge& AttributesForEdge = *new( AttributesForEdges ) FAttributesForEdge();
+								AttributesForEdge.EdgeID = NewEdgeID;
+								AttributesForEdge.EdgeAttributes.Attributes.Add( FMeshElementAttributeData( UEditableMeshAttribute::EdgeIsHard(), 0, FVector4( 1.0f ) ) );
+							}
 						}
 					}
 				}
 			}
 		}
-		check( NextAvailableExtrudedVertexIDNumber == ExtrudedVertexIDs.Num() );
+	}
+
+	for( int32 PolygonIter = 0; PolygonIter < Polygons.Num(); ++PolygonIter )
+	{
+		const FPolygonRef& PolygonRef = Polygons[ PolygonIter ];
+
+		static TArray<FVertexID> PolygonVertexIDs;
+		this->GetPolygonPerimeterVertices( PolygonRef, /* Out */ PolygonVertexIDs );
 
 		// Create a new extruded polygon for the face
 		FPolygonRef ExtrudedFrontPolygonRef;	// Filled in below
@@ -3944,7 +4065,16 @@ void UEditableMesh::ExtrudePolygons( const TArray<FPolygonRef>& Polygons, const
 			for( int32 PolygonVertexNumber = 0; PolygonVertexNumber < PolygonVertexIDs.Num(); ++PolygonVertexNumber )
 			{
 				const FVertexID VertexID = PolygonVertexIDs[ PolygonVertexNumber ];
-				NewFrontPolygonVertices[ PolygonVertexNumber ].VertexID = VertexIDToExtrudedCopyMap[ VertexID ];
+				const FVertexID* ExtrudedCopyVertexIDPtr = VertexIDToExtrudedCopy.Find( VertexID );
+				if( ExtrudedCopyVertexIDPtr != nullptr )
+				{
+					NewFrontPolygonVertices[ PolygonVertexNumber ].VertexID = VertexIDToExtrudedCopy[ VertexID ];
+				}
+				else
+				{
+					// We didn't need to extrude a new copy of this vertex (because it was part of a shared edge), so just connect the polygon to the original vertex
+					NewFrontPolygonVertices[ PolygonVertexNumber ].VertexID = VertexID;
+				}
 			}
 
 			{
@@ -3998,24 +4128,34 @@ void UEditableMesh::ExtrudePolygons( const TArray<FPolygonRef>& Polygons, const
 			}
 
 
-			// All of this edges of the new polygon will be hard
+			// All of the border edges of the new polygon will be hard.  If it was a shared edge, then we'll just preserve whatever was
+			// originally going on with the internal edge.
 			{
-				const int32 PerimeterEdgeCount = this->GetPolygonPerimeterEdgeCount( ExtrudedFrontPolygonRef );
-				for( int32 PerimeterEdgeNumber = 0; PerimeterEdgeNumber < PerimeterEdgeCount; ++PerimeterEdgeNumber )
+				const int32 NewPerimeterEdgeCount = this->GetPolygonPerimeterEdgeCount( ExtrudedFrontPolygonRef );
+				check( NewPerimeterEdgeCount == GetPolygonPerimeterEdgeCount( PolygonRef ) );	// New polygon should always have the same number of edges (in the same order) as the original!
+				for( int32 PerimeterEdgeNumber = 0; PerimeterEdgeNumber < NewPerimeterEdgeCount; ++PerimeterEdgeNumber )
 				{
+					bool bOriginalEdgeWindingIsReversedForPolygon;
+					const FEdgeID OriginalEdgeID = this->GetPolygonPerimeterEdge( PolygonRef, PerimeterEdgeNumber, /* Out */ bOriginalEdgeWindingIsReversedForPolygon );
+					const bool bIsSharedEdge = bKeepNeighborsTogether && EdgeUsageCounts[ OriginalEdgeID ] > 1;
+
 					bool bEdgeWindingIsReversedForPolygon;
 					const FEdgeID EdgeID = this->GetPolygonPerimeterEdge( ExtrudedFrontPolygonRef, PerimeterEdgeNumber, /* Out */ bEdgeWindingIsReversedForPolygon );
 
+					const FVector4 NewEdgeHardnessAttribute = bIsSharedEdge ? GetEdgeAttribute( OriginalEdgeID, UEditableMeshAttribute::EdgeIsHard(), 0 ) : FVector( 1.0f );
+					const FVector4 NewEdgeCreaseSharpnessAttribute = bIsSharedEdge ? GetEdgeAttribute( OriginalEdgeID, UEditableMeshAttribute::EdgeCreaseSharpness(), 0 ) : FVector( 1.0f );
+
 					FAttributesForEdge& AttributesForEdge = *new( AttributesForEdges ) FAttributesForEdge();
 					AttributesForEdge.EdgeID = EdgeID;
-					AttributesForEdge.EdgeAttributes.Attributes.Add( FMeshElementAttributeData( UEditableMeshAttribute::EdgeIsHard(), 0, FVector4( 1.0f ) ) );
-					AttributesForEdge.EdgeAttributes.Attributes.Add( FMeshElementAttributeData( UEditableMeshAttribute::EdgeCreaseSharpness(), 0, FVector4( 1.0f ) ) );
+					AttributesForEdge.EdgeAttributes.Attributes.Add( FMeshElementAttributeData( UEditableMeshAttribute::EdgeIsHard(), 0, NewEdgeHardnessAttribute ) );
+					AttributesForEdge.EdgeAttributes.Attributes.Add( FMeshElementAttributeData( UEditableMeshAttribute::EdgeCreaseSharpness(), 0, NewEdgeCreaseSharpnessAttribute ) );
 				}
 			}
 		}
 
 		OutNewExtrudedFrontPolygons.Add( ExtrudedFrontPolygonRef );
 	}
+	check( NextAvailableExtrudedVertexIDNumber == ExtrudedVertexIDs.Num() );	// Make sure all of the vertices we created were actually used by new polygons
 
 	// Update edge attributes in bulk
 	SetEdgesAttributes( AttributesForEdges );
diff --git a/Engine/Source/Runtime/MeshEditingRuntime/Public/EditableMesh.h b/Engine/Source/Runtime/MeshEditingRuntime/Public/EditableMesh.h
--- a/Engine/Source/Runtime/MeshEditingRuntime/Public/EditableMesh.h
+++ b/Engine/Source/Runtime/MeshEditingRuntime/Public/EditableMesh.h
@@ -291,7 +291,7 @@ public:
 	UFUNCTION( BlueprintCallable, Category="Editable Mesh" ) void SetPolygonsVertexAttributes( const TArray<FVertexAttributesForPolygon>& VertexAttributesForPolygons );
 	UFUNCTION( BlueprintCallable, Category="Editable Mesh" ) void TryToRemovePolygonEdge( const FEdgeID EdgeID, bool& bOutWasEdgeRemoved, FPolygonRef& OutNewPolygonRef );
 	UFUNCTION( BlueprintCallable, Category="Editable Mesh" ) void TryToRemoveVertex( const FVertexID VertexID, bool& bOutWasVertexRemoved, FEdgeID& OutNewEdgeID );
-	UFUNCTION( BlueprintCallable, Category="Editable Mesh" ) void ExtrudePolygons( const TArray<FPolygonRef>& Polygons, const float ExtrudeDistance, TArray<FPolygonRef>& OutNewExtrudedFrontPolygons );
+	UFUNCTION( BlueprintCallable, Category="Editable Mesh" ) void ExtrudePolygons( const TArray<FPolygonRef>& Polygons, const float ExtrudeDistance, const bool bKeepNeighborsTogether, TArray<FPolygonRef>& OutNewExtrudedFrontPolygons );
 	UFUNCTION( BlueprintCallable, Category="Editable Mesh" ) void ExtendEdges( const TArray<FEdgeID>& EdgeIDs, const bool bWeldNeighbors, TArray<FEdgeID>& OutNewExtendedEdgeIDs );
 	UFUNCTION( BlueprintCallable, Category="Editable Mesh" ) void ExtendVertices( const TArray<FVertexID>& VertexIDs, const bool bOnlyExtendClosestEdge, const FVector ReferencePosition, TArray<FVertexID>& OutNewExtendedVertexIDs );
 	UFUNCTION( BlueprintCallable, Category="Editable Mesh" ) void InsetPolygons( const TArray<FPolygonRef>& PolygonRefs, const float InsetFixedDistance, const float InsetProgressTowardCenter, const EInsetPolygonsMode Mode, TArray<FPolygonRef>& OutNewCenterPolygonRefs, TArray<FPolygonRef>& OutNewSidePolygonRefs );
