Editable Mesh: Added function to get all of a vertex's connected edges

diff --git a/Engine/Source/Runtime/MeshEditingRuntime/EditableMesh.cpp b/Engine/Source/Runtime/MeshEditingRuntime/EditableMesh.cpp
--- a/Engine/Source/Runtime/MeshEditingRuntime/EditableMesh.cpp
+++ b/Engine/Source/Runtime/MeshEditingRuntime/EditableMesh.cpp
@@ -204,6 +204,20 @@ bool UEditableMesh::IsPreviewingSubdivisions() const
 }
 
 
+void UEditableMesh::GetVertexConnectedEdges( const FVertexID VertexID, TArray<FEdgeID>& OutConnectedEdgeIDs ) const
+{
+	OutConnectedEdgeIDs.Reset();
+
+	const int32 ConnectedEdgeCount = GetVertexConnectedEdgeCount( VertexID );
+	OutConnectedEdgeIDs.Reserve( ConnectedEdgeCount );
+	for( int32 EdgeNumber = 0; EdgeNumber < ConnectedEdgeCount; ++EdgeNumber )
+	{
+		const FEdgeID ConnectedEdgeID = GetVertexConnectedEdge( VertexID, EdgeNumber );
+		OutConnectedEdgeIDs.Add( ConnectedEdgeID );
+	}
+}
+
+
 void UEditableMesh::GetVertexConnectedPolygons( const FVertexID VertexID, TArray<FPolygonRef>& OutConnectedPolygonRefs ) const
 {
 	OutConnectedPolygonRefs.Reset();
diff --git a/Engine/Source/Runtime/MeshEditingRuntime/Public/EditableMesh.h b/Engine/Source/Runtime/MeshEditingRuntime/Public/EditableMesh.h
--- a/Engine/Source/Runtime/MeshEditingRuntime/Public/EditableMesh.h
+++ b/Engine/Source/Runtime/MeshEditingRuntime/Public/EditableMesh.h
@@ -234,6 +234,7 @@ public:
 	UFUNCTION( BlueprintPure, Category="Editable Mesh" ) int32 GetTextureCoordinateCount() const;
 	UFUNCTION( BlueprintPure, Category="Editable Mesh" ) int32 GetSubdivisionCount() const;
 	UFUNCTION( BlueprintPure, Category="Editable Mesh" ) bool IsPreviewingSubdivisions() const;
+	UFUNCTION( BlueprintPure, Category="Editable Mesh" ) void GetVertexConnectedEdges( const FVertexID VertexID, TArray<FEdgeID>& OutConnectedEdgeIDs ) const;
 	UFUNCTION( BlueprintPure, Category="Editable Mesh" ) void GetVertexConnectedPolygons( const FVertexID VertexID, TArray<FPolygonRef>& OutConnectedPolygonRefs ) const;
 	UFUNCTION( BlueprintPure, Category="Editable Mesh" ) void GetVertexAdjacentVertices( const FVertexID VertexID, TArray< FVertexID >& OutAdjacentVertexIDs ) const;
 	UFUNCTION( BlueprintPure, Category="Editable Mesh" ) void GetEdgeVertices( const FEdgeID EdgeID, FVertexID& OutEdgeVertexID0, FVertexID& OutEdgeVertexID1 ) const;
