Editable Mesh: Don't change crease sharpness when extruding or insetting
- We tried this for awhile.  I think it feels bad though.

diff --git a/Engine/Source/Runtime/MeshEditingRuntime/EditableMesh.cpp b/Engine/Source/Runtime/MeshEditingRuntime/EditableMesh.cpp
--- a/Engine/Source/Runtime/MeshEditingRuntime/EditableMesh.cpp
+++ b/Engine/Source/Runtime/MeshEditingRuntime/EditableMesh.cpp
@@ -3864,7 +3864,6 @@ void UEditableMesh::ExtrudePolygons( const TArray<FPolygonRef>& Polygons, const
 						FAttributesForEdge& AttributesForEdge = *new( AttributesForEdges ) FAttributesForEdge();
 						AttributesForEdge.EdgeID = EdgeID;
 						AttributesForEdge.EdgeAttributes.Attributes.Add( FMeshElementAttributeData( UEditableMeshAttribute::EdgeIsHard(), 0, FVector4( 1.0f ) ) );
-						AttributesForEdge.EdgeAttributes.Attributes.Add( FMeshElementAttributeData( UEditableMeshAttribute::EdgeCreaseSharpness(), 0, FVector4( 1.0f ) ) );
 					}
 
 					FVertexID ExtrudedEdgeVertexIDs[ 2 ];
@@ -4157,12 +4156,10 @@ void UEditableMesh::ExtrudePolygons( const TArray<FPolygonRef>& Polygons, const
 					const FEdgeID EdgeID = this->GetPolygonPerimeterEdge( ExtrudedFrontPolygonRef, PerimeterEdgeNumber, /* Out */ bEdgeWindingIsReversedForPolygon );
 
 					const FVector4 NewEdgeHardnessAttribute = bIsSharedEdge ? GetEdgeAttribute( OriginalEdgeID, UEditableMeshAttribute::EdgeIsHard(), 0 ) : FVector( 1.0f );
-					const FVector4 NewEdgeCreaseSharpnessAttribute = bIsSharedEdge ? GetEdgeAttribute( OriginalEdgeID, UEditableMeshAttribute::EdgeCreaseSharpness(), 0 ) : FVector( 1.0f );
 
 					FAttributesForEdge& AttributesForEdge = *new( AttributesForEdges ) FAttributesForEdge();
 					AttributesForEdge.EdgeID = EdgeID;
 					AttributesForEdge.EdgeAttributes.Attributes.Add( FMeshElementAttributeData( UEditableMeshAttribute::EdgeIsHard(), 0, NewEdgeHardnessAttribute ) );
-					AttributesForEdge.EdgeAttributes.Attributes.Add( FMeshElementAttributeData( UEditableMeshAttribute::EdgeCreaseSharpness(), 0, NewEdgeCreaseSharpnessAttribute ) );
 				}
 			}
 		}
@@ -4902,9 +4899,6 @@ void UEditableMesh::BevelOrInsetPolygons( const TArray<FPolygonRef>& PolygonRefs
 	static TArray<bool> EdgesNewIsHard;
 	EdgesNewIsHard.Reset();
 
-	static TArray<float> EdgesNewCreaseSharpness;
-	EdgesNewCreaseSharpness.Reset();
-
 	if( SidePolygonsToCreate.Num() > 0 )
 	{
 		static TArray<FEdgeID> NewEdgeIDs;
@@ -4929,7 +4923,6 @@ void UEditableMesh::BevelOrInsetPolygons( const TArray<FPolygonRef>& PolygonRefs
 			{
 				EdgesToMakeHard.Add( EdgeID );
 				EdgesNewIsHard.Add( true );
-				EdgesNewCreaseSharpness.Add( 1.0f );
 			}
 		}
 	}
@@ -4960,14 +4953,12 @@ void UEditableMesh::BevelOrInsetPolygons( const TArray<FPolygonRef>& PolygonRefs
 		{
 			EdgesToMakeHard.Add( EdgeID );
 			EdgesNewIsHard.Add( true );
-			EdgesNewCreaseSharpness.Add( 1.0f );
 		}
 	}
 
 	if( EdgesToMakeHard.Num() > 0 )
 	{
 		SetEdgesHardness( EdgesToMakeHard, EdgesNewIsHard );
-		SetEdgesCreaseSharpness( EdgesToMakeHard, EdgesNewCreaseSharpness );
 	}
 
 	// Generate normals for all of the new polygons
