Copying //Tasks/UE4/Dev-Geometry-Refactor to Dev-Geometry (//UE4/Dev-Geometry)

diff --git a/Engine/Source/Runtime/MeshEditingRuntime/EditableMesh.cpp b/Engine/Source/Runtime/MeshEditingRuntime/EditableMesh.cpp
--- a/Engine/Source/Runtime/MeshEditingRuntime/EditableMesh.cpp
+++ b/Engine/Source/Runtime/MeshEditingRuntime/EditableMesh.cpp
@@ -3,11 +3,12 @@
 #include "EditableMesh.h"
 #include "EditableMeshChanges.h"
 #include "EditableMeshCustomVersion.h"
+#include "EditableMeshAdapter.h"
 #include "HAL/IConsoleManager.h"
+#include "ProfilingDebugging/ScopedTimers.h"	// For FAutoScopedDurationTimer
 #include "GeomTools.h"
 #include "mikktspace.h"	// For tangent computations
 
-
 // =========================================================
 // OpenSubdiv support
 // =========================================================
@@ -30,6 +31,8 @@
 	#pragma warning(pop)
 #endif
 
+DEFINE_LOG_CATEGORY( LogEditableMesh );
+
 // =========================================================
 
 
@@ -46,10 +49,10 @@
 
 const FElementID FElementID::Invalid( TNumericLimits<uint32>::Max() );
 const FVertexID FVertexID::Invalid( TNumericLimits<uint32>::Max() );
+const FVertexInstanceID FVertexInstanceID::Invalid( TNumericLimits<uint32>::Max() );
 const FEdgeID FEdgeID::Invalid( TNumericLimits<uint32>::Max() );
-const FSectionID FSectionID::Invalid( TNumericLimits<uint32>::Max() );
+const FPolygonGroupID FPolygonGroupID::Invalid( TNumericLimits<uint32>::Max() );
 const FPolygonID FPolygonID::Invalid( TNumericLimits<uint32>::Max() );
-const FPolygonRef FPolygonRef::Invalid( FSectionID::Invalid, FPolygonID::Invalid );
 
 const FName UEditableMeshAttribute::VertexPositionName( "VertexPosition" );
 const FName UEditableMeshAttribute::VertexCornerSharpnessName( "VertexCornerSharpness" );
@@ -60,9 +63,12 @@
 const FName UEditableMeshAttribute::VertexColorName( "VertexColor" );
 const FName UEditableMeshAttribute::EdgeIsHardName( "EdgeIsHard" );
 const FName UEditableMeshAttribute::EdgeCreaseSharpnessName( "EdgeCreaseSharpness" );
+const FName UEditableMeshAttribute::PolygonNormalName( "PolygonNormal" );
+const FName UEditableMeshAttribute::PolygonCenterName( "PolygonCenter" );
 
 
 UEditableMesh::UEditableMesh()
+	: PendingCompactCounter( 0 )
 {
 }
 
@@ -71,6 +77,12 @@
 {
 	Super::Serialize( Ar );
 	Ar.UsingCustomVersion( FEditableMeshCustomVersion::GUID );
+
+	SerializeSparseArray( Ar, Vertices );
+	SerializeSparseArray( Ar, VertexInstances );
+	SerializeSparseArray( Ar, Edges );
+	SerializeSparseArray( Ar, Polygons );
+	SerializeSparseArray( Ar, PolygonGroups );
 }
 
 
@@ -81,11 +93,453 @@
 	if( IsPreviewingSubdivisions() )
 	{
 		RefreshOpenSubdiv();
-		RebuildRenderMesh();
+	}
+
+	RebuildRenderMesh();
+}
+
+
+template <typename ElementIDType>
+static void InvertRemapTable( TSparseArray<ElementIDType>& InvertedRemapTable, const TSparseArray<ElementIDType>& RemapTable )
+{
+	static_assert( TIsDerivedFrom<ElementIDType, FElementID>::IsDerived, "Remap array type must be derived from FElementID" );
+
+	InvertedRemapTable.Empty( RemapTable.Num() );
+
+	for( TSparseArray<ElementIDType>::TConstIterator It( RemapTable ); It; ++It )
+	{
+		InvertedRemapTable.Insert( It->GetValue(), ElementIDType( It.GetIndex() ) );
 	}
 }
 
 
+class FCompactChange : public FChange
+{
+public:
+
+	/** Constructor */
+	FCompactChange()
+	{
+	}
+
+	// Parent class overrides
+	virtual TUniquePtr<FChange> Execute( UObject* Object ) override
+	{
+		UEditableMesh* EditableMesh = CastChecked<UEditableMesh>( Object );
+		verify( !EditableMesh->AnyChangesToUndo() );
+		EditableMesh->Compact();
+		return EditableMesh->MakeUndo();
+	}
+
+	virtual FString ToString() const override
+	{
+		return FString( TEXT( "Compact" ) );
+	}
+};
+
+
+struct FUncompactChangeInput
+{
+	/** A set of remap tables, specifying how the elements should have their indices remapped */
+	FElementIDRemappings ElementIDRemappings;
+};
+
+
+class FUncompactChange : public FChange
+{
+public:
+
+	/** Constructor */
+	FUncompactChange( const FUncompactChangeInput& InitInput )
+		: Input( InitInput )
+	{
+	}
+
+	FUncompactChange( FUncompactChangeInput&& InitInput )
+		: Input( MoveTemp( InitInput ) )
+	{
+	}
+
+	// Parent class overrides
+	virtual TUniquePtr<FChange> Execute( UObject* Object ) override
+	{
+		UEditableMesh* EditableMesh = CastChecked<UEditableMesh>( Object );
+		verify( !EditableMesh->AnyChangesToUndo() );
+		EditableMesh->Uncompact( Input.ElementIDRemappings );
+		return EditableMesh->MakeUndo();
+	}
+
+	virtual FString ToString() const override
+	{
+		return FString( TEXT( "Uncompact" ) );
+	}
+
+private:
+
+	/** The data we need to make this change */
+	FUncompactChangeInput Input;
+};
+
+
+void UEditableMesh::FixUpElementIDs( const FElementIDRemappings& Remappings )
+{
+	for( FMeshVertex& Vertex : Vertices )
+	{
+		// Fix up vertex instance index references in vertices array
+		for( FVertexInstanceID& VertexInstanceID : Vertex.VertexInstanceIDs )
+		{
+			VertexInstanceID = Remappings.GetRemappedVertexInstanceID( VertexInstanceID );
+		}
+
+		// Fix up edge index references in the vertex array
+		for( FEdgeID& EdgeID : Vertex.ConnectedEdgeIDs )
+		{
+			EdgeID = Remappings.GetRemappedEdgeID( EdgeID );
+		}
+	}
+
+	// Fix up vertex index references in vertex instance array
+	for( FMeshVertexInstance& VertexInstance : VertexInstances )
+	{
+		VertexInstance.VertexID = Remappings.GetRemappedVertexID( VertexInstance.VertexID );
+
+		for( FPolygonID& PolygonID : VertexInstance.ConnectedPolygons )
+		{
+			PolygonID = Remappings.GetRemappedPolygonID( PolygonID );
+		}
+	}
+
+	for( FMeshEdge& Edge : Edges )
+	{
+		// Fix up vertex index references in Edges array
+		for( int32 Index = 0; Index < 2; Index++ )
+		{
+			Edge.VertexIDs[ Index ] = Remappings.GetRemappedVertexID( Edge.VertexIDs[ Index ] );
+		}
+
+		// Fix up references to section indices
+		for( FPolygonID& ConnectedPolygon : Edge.ConnectedPolygons )
+		{
+			ConnectedPolygon = Remappings.GetRemappedPolygonID( ConnectedPolygon );
+		}
+	}
+
+	for( FMeshPolygon& Polygon : Polygons )
+	{
+		// Fix up references to vertex indices in section polygons' contours
+		for( FVertexInstanceID& VertexInstanceID : Polygon.PerimeterContour.VertexInstanceIDs )
+		{
+			VertexInstanceID = Remappings.GetRemappedVertexInstanceID( VertexInstanceID );
+		}
+
+		for( FMeshPolygonContour& HoleContour : Polygon.HoleContours )
+		{
+			for( FVertexInstanceID& VertexInstanceID : HoleContour.VertexInstanceIDs )
+			{
+				VertexInstanceID = Remappings.GetRemappedVertexInstanceID( VertexInstanceID );
+			}
+		}
+
+		for( FMeshTriangle& Triangle : Polygon.Triangles )
+		{
+			for( FVertexInstanceID& VertexInstanceID : Triangle.VertexInstanceIDs )
+			{
+				VertexInstanceID = Remappings.GetRemappedVertexInstanceID( VertexInstanceID );
+			}
+		}
+
+		Polygon.PolygonGroupID = Remappings.GetRemappedPolygonGroupID( Polygon.PolygonGroupID );
+	}
+
+	for( FMeshPolygonGroup& PolygonGroup : PolygonGroups )
+	{
+		for( FPolygonID& Polygon : PolygonGroup.Polygons )
+		{
+			Polygon = Remappings.GetRemappedPolygonID( Polygon );
+		}
+	}
+}
+
+
+void UEditableMesh::Compact()
+{
+	static FElementIDRemappings Remappings;
+
+	CompactSparseArrayElements( Vertices, Remappings.NewVertexIndexLookup );
+	CompactSparseArrayElements( VertexInstances, Remappings.NewVertexInstanceIndexLookup );
+	CompactSparseArrayElements( Edges, Remappings.NewEdgeIndexLookup );
+	CompactSparseArrayElements( Polygons, Remappings.NewPolygonIndexLookup );
+	CompactSparseArrayElements( PolygonGroups, Remappings.NewPolygonGroupIndexLookup );
+
+	FixUpElementIDs( Remappings );
+
+	for( UEditableMeshAdapter* Adapter : Adapters )
+	{
+		Adapter->OnReindexElements( this, Remappings );
+	}
+
+	// @todo mesheditor: broadcast event with remappings so that any cached element IDs can be fixed up.
+
+	RebuildRenderMesh();
+
+	// Prepare the inverse transaction to reverse the compaction
+	FUncompactChangeInput UncompactChangeInput;
+	InvertRemapTable( UncompactChangeInput.ElementIDRemappings.NewVertexIndexLookup, Remappings.NewVertexIndexLookup );
+	InvertRemapTable( UncompactChangeInput.ElementIDRemappings.NewVertexInstanceIndexLookup, Remappings.NewVertexInstanceIndexLookup );
+	InvertRemapTable( UncompactChangeInput.ElementIDRemappings.NewEdgeIndexLookup, Remappings.NewEdgeIndexLookup );
+	InvertRemapTable( UncompactChangeInput.ElementIDRemappings.NewPolygonIndexLookup, Remappings.NewPolygonIndexLookup );
+	InvertRemapTable( UncompactChangeInput.ElementIDRemappings.NewPolygonGroupIndexLookup, Remappings.NewPolygonGroupIndexLookup );
+
+	AddUndo( MakeUnique<FUncompactChange>( MoveTemp( UncompactChangeInput ) ) );
+}
+
+
+void UEditableMesh::Uncompact( const FElementIDRemappings& Remappings )
+{
+	RemapSparseArrayElements( Vertices, Remappings.NewVertexIndexLookup );
+	RemapSparseArrayElements( VertexInstances, Remappings.NewVertexInstanceIndexLookup );
+	RemapSparseArrayElements( Edges, Remappings.NewEdgeIndexLookup );
+	RemapSparseArrayElements( Polygons, Remappings.NewPolygonIndexLookup );
+	RemapSparseArrayElements( PolygonGroups, Remappings.NewPolygonGroupIndexLookup );
+
+	FixUpElementIDs( Remappings );
+
+	for( UEditableMeshAdapter* Adapter : Adapters )
+	{
+		Adapter->OnReindexElements( this, Remappings );
+	}
+
+	// @todo mesheditor: broadcast event with remappings so that any cached element IDs can be fixed up.
+
+	RebuildRenderMesh();
+
+	AddUndo( MakeUnique<FCompactChange>() );
+}
+
+
+void UEditableMesh::RebuildRenderMesh()
+{
+	if( !IsBeingModified() )
+	{
+		const bool bRefreshBounds = true;
+		const bool bInvalidateLighting = true;
+		for( UEditableMeshAdapter* Adapter : Adapters )
+		{
+			Adapter->OnRebuildRenderMeshStart( this, bRefreshBounds, bInvalidateLighting );
+		}
+	}
+
+	for( UEditableMeshAdapter* Adapter : Adapters )
+	{
+		Adapter->OnRebuildRenderMesh( this );
+	}
+
+	if( !IsBeingModified() )
+	{
+		const bool bUpdateCollision = true;
+		for( UEditableMeshAdapter* Adapter : Adapters )
+		{
+			Adapter->OnRebuildRenderMeshFinish( this, bUpdateCollision );
+		}
+	}
+}
+
+
+void UEditableMesh::StartModification( const EMeshModificationType MeshModificationType, const EMeshTopologyChange MeshTopologyChange )
+{
+	if( ensure( !IsBeingModified() ) )
+	{
+		bIsBeingModified = true;
+
+		// Should be nothing in the undo stack if we're just starting to modify the mesh now
+		ensure( !this->AnyChangesToUndo() );
+
+		FStartOrEndModificationChangeInput RevertInput;
+		RevertInput.bStartModification = false;
+		RevertInput.MeshModificationType = MeshModificationType;
+		RevertInput.MeshTopologyChange = MeshTopologyChange;
+		AddUndo( MakeUnique<FStartOrEndModificationChange>( MoveTemp( RevertInput ) ) );
+
+		this->CurrentModificationType = MeshModificationType;
+		this->CurrentToplogyChange = MeshTopologyChange;
+
+		// @todo mesheditor debug: Disable noisy mesh editor spew by default (here and elsewhere)
+		// UE_LOG( LogMeshEditingRuntime, Log, TEXT( "UEditableStaticMesh::StartModification START: %s" ), *SubMeshAddress.ToString() );
+		FAutoScopedDurationTimer FunctionTimer;
+
+		const bool bRefreshBounds = MeshModificationType == EMeshModificationType::Final;	 // @todo mesheditor perf: Only do this if we may have changed the bounds
+		const bool bInvalidateLighting = ( MeshModificationType == EMeshModificationType::FirstInterim || MeshModificationType == EMeshModificationType::Final );	// @todo mesheditor perf: We can avoid invalidating lighting on 'Final' if we know that a 'FirstInterim' happened since the last 'Final'
+		for( UEditableMeshAdapter* Adapter : Adapters )
+		{
+			Adapter->OnRebuildRenderMeshStart( this, bRefreshBounds, bInvalidateLighting );
+		}
+
+		for( UEditableMeshAdapter* Adapter : Adapters )
+		{
+			Adapter->OnStartModification( this, MeshModificationType, MeshTopologyChange );
+		}
+
+		PolygonsPendingNewTangentBasis.Reset();
+		PolygonsPendingTriangulation.Reset();
+
+		// @todo mesheditor debug
+		// UE_LOG( LogMeshEditingRuntime, Log, TEXT( "UEditableStaticMesh::StartModification COMPLETE in %0.4fs" ), FunctionTimer.GetTime() );
+	}
+}
+
+
+void UEditableMesh::EndModification( const bool bFromUndo )
+{
+	if( ensure( IsBeingModified() ) )
+	{
+		// @todo mesheditor debug
+		// UE_LOG( LogMeshEditingRuntime, Log, TEXT( "UEditableStaticMesh::EndModification START (ModType=%i): %s" ), (int32)MeshModificationType, *SubMeshAddress.ToString() );
+		// FAutoScopedDurationTimer FunctionTimer;
+
+		// Retriangulate first, as the triangulation of n-gons determines how the tangent basis is calculated
+		if( PolygonsPendingTriangulation.Num() > 0 )
+		{
+			RetriangulatePolygons();
+		}
+
+		if( PolygonsPendingNewTangentBasis.Num() > 0 )
+		{
+			GenerateTangentsAndNormals();
+		}
+
+		if( CurrentModificationType == EMeshModificationType::Final || !bFromUndo )
+		{
+			// Update subdivision limit surface
+			if( CurrentToplogyChange == EMeshTopologyChange::TopologyChange )
+			{
+				// Mesh topology (or subdivision level or smoothing) may have changed, so go ahead and refresh our OpenSubdiv representation entirely
+				RefreshOpenSubdiv();
+			}
+			else
+			{
+				// No topology change, so we can ask OpenSubdiv to quickly generate new limit surface geometry
+				GenerateOpenSubdivLimitSurfaceData();
+			}
+		}
+
+		// Every so often, compact the data.
+		// Note we only want to do this when actions are performed, not when they are being undone/redone
+		bool bDidCompact = false;
+
+		if( CurrentModificationType == EMeshModificationType::Final &&
+			CurrentToplogyChange == EMeshTopologyChange::TopologyChange &&
+			!bFromUndo )
+		{
+			if( ++PendingCompactCounter == CompactFrequency )
+			{
+				PendingCompactCounter = 0;
+				Compact();
+				bDidCompact = true;
+			}
+		}
+
+		// If subdivision preview mode is active, we'll need to refresh the entire static mesh with data from OpenSubdiv
+		// @todo mesheditor subdiv perf: Ideally we can avoid refreshing the entire thing if only positions have changed, as per above
+		if( IsPreviewingSubdivisions() && ( CurrentModificationType == EMeshModificationType::Final || !bFromUndo ) )
+		{
+			if( !bDidCompact )	// If we did a Compact() in this function, the mesh will have already been rebuilt
+			{
+				for( UEditableMeshAdapter* Adapter : Adapters )
+				{
+					Adapter->OnRebuildRenderMesh( this );
+				}
+			}
+		}
+
+		for( UEditableMeshAdapter* Adapter : Adapters )
+		{
+			Adapter->OnRebuildRenderMeshFinish( this, CurrentModificationType == EMeshModificationType::Final );
+		}
+
+		for( UEditableMeshAdapter* Adapter : Adapters )
+		{
+			Adapter->OnEndModification( this );
+		}
+
+		// @todo mesheditor: Not currently sure if we need to do this or not.  Also, we are trying to support runtime editing (no PostEditChange stuff)
+		// 		if( GIsEditor && MeshModificationType == EMeshModificationType::Final )
+		// 		{
+		// 			StaticMesh->PostEditChange();
+		// 		}
+
+		// @todo mesheditor debug
+		// UE_LOG( LogMeshEditingRuntime, Log, TEXT( "UEditableStaticMesh::EndModification COMPLETE in %0.4fs" ), FunctionTimer.GetTime() );	  // @todo mesheditor: Shows bogus time values
+
+		FStartOrEndModificationChangeInput RevertInput;
+		RevertInput.bStartModification = true;
+		RevertInput.MeshModificationType = CurrentModificationType;
+		RevertInput.MeshTopologyChange = CurrentToplogyChange;
+		AddUndo( MakeUnique<FStartOrEndModificationChange>( MoveTemp( RevertInput ) ) );
+
+		bIsBeingModified = false;
+	}
+}
+
+
+bool UEditableMesh::IsCommitted() const
+{
+	for( UEditableMeshAdapter* Adapter : Adapters )
+	{
+		if( !Adapter->IsCommitted( this ) )
+		{
+			return false;
+		}
+	}
+
+	return true;
+}
+
+
+bool UEditableMesh::IsCommittedAsInstance() const
+{
+	for( UEditableMeshAdapter* Adapter : Adapters )
+	{
+		if( !Adapter->IsCommittedAsInstance( this ) )
+		{
+			return false;
+		}
+	}
+
+	return true;
+}
+
+
+void UEditableMesh::Commit()
+{
+	for( UEditableMeshAdapter* Adapter : Adapters )
+	{
+		Adapter->OnCommit( this );
+	}
+}
+
+
+UEditableMesh* UEditableMesh::CommitInstance( UPrimitiveComponent* ComponentToInstanceTo )
+{
+	UEditableMesh* Result = nullptr;
+	for( UEditableMeshAdapter* Adapter : Adapters )
+	{
+		UEditableMesh* Instance = Adapter->OnCommitInstance( this, ComponentToInstanceTo );
+		if( !Result )
+		{
+			Result = Instance;
+		}
+	}
+
+	return Result;
+}
+
+
+// @todo mesheditor: implement these, or - better - reconsider the editing paradigm
+void UEditableMesh::Revert() {}
+UEditableMesh* UEditableMesh::RevertInstance() { return nullptr; }
+void UEditableMesh::PropagateInstanceChanges() {}
+
+
 const FEditableMeshSubMeshAddress& UEditableMesh::GetSubMeshAddress() const
 {
 	return SubMeshAddress;
@@ -98,6 +552,775 @@
 }
 
 
+int32 UEditableMesh::GetVertexCount() const
+{
+	return Vertices.Num();
+}
+
+
+int32 UEditableMesh::GetVertexArraySize() const
+{
+	return Vertices.GetMaxIndex();
+}
+
+
+bool UEditableMesh::IsValidVertex( const FVertexID VertexID ) const
+{
+	return VertexID.GetValue() >= 0 && VertexID.GetValue() < Vertices.GetMaxIndex() && Vertices.IsAllocated( VertexID.GetValue() );
+}
+
+
+FVector4 UEditableMesh::GetVertexAttribute( const FVertexID VertexID, const FName AttributeName, const int32 AttributeIndex ) const
+{
+	const FMeshVertex& Vertex = Vertices[ VertexID.GetValue() ];
+
+	if( AttributeName == UEditableMeshAttribute::VertexPosition() )
+	{
+		check( AttributeIndex == 0 );	// Only one position is supported
+		return FVector4( Vertex.VertexPosition, 0.0f );
+	}
+	else if( AttributeName == UEditableMeshAttribute::VertexCornerSharpness() )
+	{
+		check( AttributeIndex == 0 );	// Only one softness value is supported
+		return FVector4( Vertex.CornerSharpness, 0.0f, 0.0f, 0.0f );
+	}
+
+	checkf( 0, TEXT( "UEditableMesh::GetVertexAttribute() called with unrecognized vertex attribute name: %s (index: %i)" ), *AttributeName.ToString(), AttributeIndex );
+	return FVector4( 0.0f );
+}
+
+
+void UEditableMesh::SetVertexAttribute( const FVertexID VertexID, const FName AttributeName, const int32 AttributeIndex, const FVector4 AttributeValue )
+{
+	FMeshVertex& Vertex = Vertices[ VertexID.GetValue() ];
+
+	if( AttributeName == UEditableMeshAttribute::VertexPosition() )
+	{
+		check( AttributeIndex == 0 );	// Only one position is supported
+		Vertex.VertexPosition = AttributeValue;
+	}
+	else if( AttributeName == UEditableMeshAttribute::VertexCornerSharpness() )
+	{
+		check( AttributeIndex == 0 );	// Only one softness value is supported
+		Vertex.CornerSharpness = AttributeValue.X;
+	}
+	else
+	{
+		checkf( 0, TEXT( "UEditableMesh::SetVertexAttribute() called with unrecognized vertex attribute name: %s (index: %i)" ), *AttributeName.ToString(), AttributeIndex );
+	}
+
+	for( UEditableMeshAdapter* Adapter : Adapters )
+	{
+		Adapter->OnSetVertexAttribute( this, VertexID, AttributeName, AttributeIndex, AttributeValue );
+	}
+}
+
+
+int32 UEditableMesh::GetVertexConnectedEdgeCount( const FVertexID VertexID ) const
+{
+	checkSlow( Vertices.IsAllocated( VertexID.GetValue() ) );
+	const FMeshVertex& Vertex = Vertices[ VertexID.GetValue() ];
+	return Vertex.ConnectedEdgeIDs.Num();
+}
+
+
+FEdgeID UEditableMesh::GetVertexConnectedEdge( const FVertexID VertexID, const int32 ConnectedEdgeNumber ) const
+{
+	checkSlow( Vertices.IsAllocated( VertexID.GetValue() ) );
+	const FMeshVertex& Vertex = Vertices[ VertexID.GetValue() ];
+	return Vertex.ConnectedEdgeIDs[ ConnectedEdgeNumber ];
+}
+
+
+int32 UEditableMesh::GetVertexInstanceCount() const
+{
+	return VertexInstances.Num();
+}
+
+
+int32 UEditableMesh::GetVertexInstanceArraySize() const
+{
+	return VertexInstances.GetMaxIndex();
+}
+
+
+FVector4 UEditableMesh::GetVertexInstanceAttribute( const FVertexInstanceID VertexInstanceID, const FName AttributeName, const int32 AttributeIndex ) const
+{
+	const FMeshVertexInstance& VertexInstance = VertexInstances[ VertexInstanceID.GetValue() ];
+
+	if( AttributeName == UEditableMeshAttribute::VertexPosition() )
+	{
+		check( AttributeIndex == 0 );	// Only one position is supported
+		return GetVertexAttribute( VertexInstance.VertexID, AttributeName, AttributeIndex );
+	}
+	else if( AttributeName == UEditableMeshAttribute::VertexNormal() )
+	{
+		check( AttributeIndex == 0 );	// Only one normal is supported
+		return FVector4( VertexInstance.Normal, 0.0f );
+	}
+	else if( AttributeName == UEditableMeshAttribute::VertexTangent() )
+	{
+		check( AttributeIndex == 0 );	// Only one tangent is supported
+		return FVector4( VertexInstance.Tangent, 0.0f );
+	}
+	else if( AttributeName == UEditableMeshAttribute::VertexBinormalSign() )
+	{
+		check( AttributeIndex == 0 );	// Only one basis determinant sign is supported
+		return FVector4( VertexInstance.BinormalSign );
+	}
+	else if( AttributeName == UEditableMeshAttribute::VertexTextureCoordinate() )
+	{
+		const int32 TextureCoordinateIndex = AttributeIndex;
+		if( TextureCoordinateIndex < VertexInstance.VertexUVs.Num() )
+		{
+			const FVector2D TextureCoordinate = VertexInstance.VertexUVs[ TextureCoordinateIndex ];
+			return FVector4( TextureCoordinate, FVector2D::ZeroVector );
+		}
+		else
+		{
+			return FVector4( 0.0f );
+		}
+	}
+	else if( AttributeName == UEditableMeshAttribute::VertexColor() )
+	{
+		check( AttributeIndex == 0 );
+		return FVector4( FLinearColor( VertexInstance.Color ) );
+	}
+
+	checkf( 0, TEXT( "UEditableMesh::GetVertexInstanceAttribute() called with unrecognized vertex instance attribute name: %s (index: %i)" ), *AttributeName.ToString(), AttributeIndex );
+	return FVector4( 0.0f );
+}
+
+
+void UEditableMesh::SetVertexInstanceAttribute( const FVertexInstanceID VertexInstanceID, const FName AttributeName, const int32 AttributeIndex, const FVector4 AttributeValue )
+{
+	FMeshVertexInstance& VertexInstance = VertexInstances[ VertexInstanceID.GetValue() ];
+
+	if( AttributeName == UEditableMeshAttribute::VertexPosition() )
+	{
+		check( AttributeIndex == 0 );	// Only one position is supported
+		SetVertexAttribute( VertexInstance.VertexID, AttributeName, AttributeIndex, AttributeValue );
+	}
+	else if( AttributeName == UEditableMeshAttribute::VertexNormal() )
+	{
+		check( AttributeIndex == 0 );	// Only one normal is supported
+		VertexInstance.Normal = AttributeValue;
+	}
+	else if( AttributeName == UEditableMeshAttribute::VertexTangent() )
+	{
+		check( AttributeIndex == 0 );	// Only one tangent is supported
+		VertexInstance.Tangent = AttributeValue;
+	}
+	else if( AttributeName == UEditableMeshAttribute::VertexBinormalSign() )
+	{
+		check( AttributeIndex == 0 );	// Only one basis determinant sign is supported
+		VertexInstance.BinormalSign = AttributeValue.X;
+	}
+	else if( AttributeName == UEditableMeshAttribute::VertexTextureCoordinate() )
+	{
+		if( AttributeIndex < VertexInstance.VertexUVs.Num() )
+		{
+			VertexInstance.VertexUVs[ AttributeIndex ] = FVector2D( AttributeValue.X, AttributeValue.Y );
+		}
+	}
+	else if( AttributeName == UEditableMeshAttribute::VertexColor() )
+	{
+		check( AttributeIndex == 0 );
+		VertexInstance.Color = FLinearColor( AttributeValue.X, AttributeValue.Y, AttributeValue.Z, AttributeValue.W );
+	}
+	else
+	{
+		checkf( 0, TEXT( "UEditableMesh::SetVertexInstanceAttribute() called with unrecognized vertex instance attribute name: %s (index: %i)" ), *AttributeName.ToString(), AttributeIndex );
+	}
+
+	for( UEditableMeshAdapter* Adapter : Adapters )
+	{
+		Adapter->OnSetVertexInstanceAttribute( this, VertexInstanceID, AttributeName, AttributeIndex, AttributeValue );
+	}
+}
+
+
+FVector4 UEditableMesh::GetEdgeAttribute( const FEdgeID EdgeID, const FName AttributeName, const int32 AttributeIndex ) const
+{
+	checkSlow( Edges.IsAllocated( EdgeID.GetValue() ) );
+	const FMeshEdge& Edge = Edges[ EdgeID.GetValue() ];
+
+	if( AttributeName == UEditableMeshAttribute::EdgeIsHard() )
+	{
+		check( AttributeIndex == 0 );	// Only one edge is hard flag is supported
+
+		return FVector4( Edge.bIsHardEdge ? 1.0f : 0.0f, 0.0f, 0.0f, 0.0f );
+	}
+	else if( AttributeName == UEditableMeshAttribute::EdgeCreaseSharpness() )
+	{
+		check( AttributeIndex == 0 );	// Only one edge crease sharpness is supported
+
+		return FVector4( Edge.CreaseSharpness, 0.0f, 0.0f, 0.0f );
+	}
+
+	checkf( 0, TEXT( "UEditableMesh::GetEdgeAttribute() called with unrecognized edge attribute name: %s (index: %i)" ), *AttributeName.ToString(), AttributeIndex );
+	return FVector4( 0.0f );
+}
+
+
+void UEditableMesh::SetEdgeAttribute( const FEdgeID EdgeID, const FName AttributeName, const int32 AttributeIndex, const FVector4 AttributeValue, bool bIsUndo )
+{
+	checkSlow( Edges.IsAllocated( EdgeID.GetValue() ) );
+	FMeshEdge& Edge = Edges[ EdgeID.GetValue() ];
+
+	if( AttributeName == UEditableMeshAttribute::EdgeIsHard() )
+	{
+		check( AttributeIndex == 0 );	// Only one edge is hard flag is supported
+		SetEdgeHardness( EdgeID, AttributeValue.X != 0.0f, bIsUndo );
+
+		// @todo mesheditor: this needs to potentially merge/split vertex instances and rebuild normals
+	}
+	else if( AttributeName == UEditableMeshAttribute::EdgeCreaseSharpness() )
+	{
+		check( AttributeIndex == 0 );	// Only one edge crease sharpness is supported
+		Edge.CreaseSharpness = AttributeValue.X;
+	}
+	else
+	{
+		checkf( 0, TEXT( "UEditableMesh::SetEdgeAttribute() called with unrecognized edge attribute name: %s (index: %i)" ), *AttributeName.ToString(), AttributeIndex );
+	}
+
+	for( UEditableMeshAdapter* Adapter : Adapters )
+	{
+		Adapter->OnSetEdgeAttribute( this, EdgeID, AttributeName, AttributeIndex, AttributeValue );
+	}
+}
+
+
+void UEditableMesh::GetPolygonsInSameSoftEdgedGroup( const FVertexID VertexID, const FPolygonID PolygonID, TArray<FPolygonID>& OutPolygonIDs ) const
+{
+	// The aim here is to determine which polygons form part of the same soft edged group as the polygons attached to this vertex instance.
+	// They should all contribute to the final vertex instance normal.
+
+	OutPolygonIDs.Reset();
+
+	// Get all polygons connected to this vertex.
+	static TArray<FPolygonID> ConnectedPolygons;
+	GetVertexConnectedPolygons( VertexID, ConnectedPolygons );
+
+	// Cache a list of all soft edges which share this vertex.
+	// We're only interested in finding adjacent polygons which are not the other side of a hard edge.
+	static TArray<FEdgeID> ConnectedSoftEdges;
+	GetConnectedSoftEdges( VertexID, ConnectedSoftEdges );
+
+	// Iterate through adjacent polygons which contain the given vertex, but don't cross a hard edge.
+	// Maintain a list of polygon IDs to be examined. Adjacents are added to the list if suitable.
+	// Add the start poly here.
+	static TArray<FPolygonID> PolygonsToCheck;
+	PolygonsToCheck.Reset();
+	PolygonsToCheck.Add( PolygonID );
+
+	int32 Index = 0;
+	while( Index < PolygonsToCheck.Num() )
+	{
+		const FPolygonID PolygonToCheck = PolygonsToCheck[ Index ];
+		Index++;
+
+		if( ConnectedPolygons.Contains( PolygonToCheck ) )
+		{
+			OutPolygonIDs.Add( PolygonToCheck );
+
+			// Now look at its adjacent polygons. If they are joined by a soft edge which includes the vertex we're interested in, we want to consider them.
+			// We take a shortcut by doing this process in reverse: we already know all the soft edges we are interested in, so check if any of them
+			// have the current polygon as an adjacent.
+			for( const FEdgeID ConnectedSoftEdge : ConnectedSoftEdges )
+			{
+				const FMeshEdge& Edge = Edges[ ConnectedSoftEdge.GetValue() ];
+				if( Edge.ConnectedPolygons.Contains( PolygonToCheck ) )
+				{
+					for( const FPolygonID AdjacentPolygon : Edge.ConnectedPolygons )
+					{
+						// Only add new polygons which haven't yet been added to the list. This prevents circular runs of polygons triggering infinite loops.
+						PolygonsToCheck.AddUnique( AdjacentPolygon );
+					}
+				}
+			}
+		}
+	}
+}
+
+
+void UEditableMesh::SplitVertexInstanceInPolygons( const FVertexInstanceID VertexInstanceID, const TArray<FPolygonID>& PolygonIDs )
+{
+	UE_LOG( LogEditableMesh, Verbose, TEXT( "SplitVertexInstanceInPolygons: %s %s" ), *VertexInstanceID.ToString(), *LogHelpers::ArrayToString( PolygonIDs ) );
+
+	const FMeshVertexInstance& VertexInstance = VertexInstances[ VertexInstanceID.GetValue() ];
+
+	// Create a new vertex instance copied from the one passed in
+
+	static TArray<FVertexInstanceToCreate> VertexInstancesToCreate;
+	VertexInstancesToCreate.Reset();
+	VertexInstancesToCreate.SetNum( 1 );
+	VertexInstancesToCreate[ 0 ].VertexID = VertexInstance.VertexID;
+
+	for( const FName AttributeName : UEditableMesh::GetValidVertexInstanceAttributes() )
+	{
+		const int32 MaxAttributeIndex = GetMaxAttributeIndex( AttributeName );
+		for( int32 AttributeIndex = 0; AttributeIndex < MaxAttributeIndex; ++AttributeIndex )
+		{
+			VertexInstancesToCreate[ 0 ].VertexInstanceAttributes.Attributes.Emplace(
+				AttributeName,
+				AttributeIndex,
+				GetVertexInstanceAttribute( VertexInstanceID, AttributeName, AttributeIndex )
+			);
+		}
+	}
+
+	static TArray<FVertexInstanceID> NewVertexInstanceIDs;
+	CreateVertexInstances( VertexInstancesToCreate, NewVertexInstanceIDs );
+
+	ReplaceVertexInstanceInPolygons( VertexInstanceID, NewVertexInstanceIDs[ 0 ], PolygonIDs );
+
+	UE_LOG( LogEditableMesh, Verbose, TEXT( "* SplitVertexInstanceInPolygons returned" ) );
+}
+
+
+void UEditableMesh::ReplaceVertexInstanceInPolygons( const FVertexInstanceID OldVertexInstanceID, const FVertexInstanceID NewVertexInstanceID, const TArray<FPolygonID>& PolygonIDs )
+{
+	UE_LOG( LogEditableMesh, Verbose, TEXT( "ReplaceVertexInstanceInPolygons: %s %s %s" ), *OldVertexInstanceID.ToString(), *NewVertexInstanceID.ToString(), *LogHelpers::ArrayToString( PolygonIDs ) );
+
+	// Substitute the new vertex instance in the passed in polygons
+
+	static TArray<FChangeVertexInstancesForPolygon>  VertexInstancesToChange;
+	VertexInstancesToChange.Reset( PolygonIDs.Num() );
+
+	for( const FPolygonID PolygonID : PolygonIDs )
+	{
+		VertexInstancesToChange.Emplace();
+		FChangeVertexInstancesForPolygon& ChangeVertexInstances = VertexInstancesToChange.Last();
+
+		ChangeVertexInstances.PolygonID = PolygonID;
+
+		const FMeshPolygon& Polygon = Polygons[ PolygonID.GetValue() ];
+		int32 VertexInstanceIndex = Polygon.PerimeterContour.VertexInstanceIDs.Find( OldVertexInstanceID );
+
+		if( VertexInstanceIndex != INDEX_NONE )
+		{
+			// Found the vertex instance in the perimeter
+			ChangeVertexInstances.PerimeterVertexIndicesAndInstanceIDs.Emplace();
+			ChangeVertexInstances.PerimeterVertexIndicesAndInstanceIDs[ 0 ].ContourIndex = VertexInstanceIndex;
+			ChangeVertexInstances.PerimeterVertexIndicesAndInstanceIDs[ 0 ].VertexInstanceID = NewVertexInstanceID;
+		}
+		else
+		{
+			ChangeVertexInstances.VertexIndicesAndInstanceIDsForEachHole.SetNum( Polygon.HoleContours.Num() );
+
+			int32 HoleIndex = 0;
+			for( const FMeshPolygonContour& HoleContour : Polygon.HoleContours )
+			{
+				VertexInstanceIndex = HoleContour.VertexInstanceIDs.Find( OldVertexInstanceID );
+				if( VertexInstanceIndex != INDEX_NONE )
+				{
+					// Found the vertex instance in a hole
+					ChangeVertexInstances.VertexIndicesAndInstanceIDsForEachHole[ HoleIndex ].VertexIndicesAndInstanceIDs.Emplace();
+					ChangeVertexInstances.VertexIndicesAndInstanceIDsForEachHole[ HoleIndex ].VertexIndicesAndInstanceIDs[ 0 ].ContourIndex = VertexInstanceIndex;
+					ChangeVertexInstances.VertexIndicesAndInstanceIDsForEachHole[ HoleIndex ].VertexIndicesAndInstanceIDs[ 0 ].VertexInstanceID = NewVertexInstanceID;
+					break;
+				}
+
+				HoleIndex++;
+			}
+		}
+
+		// We expect to have found the vertex instance somewhere in one of the polygon contours
+		check( VertexInstanceIndex != INDEX_NONE );
+	}
+
+	ChangePolygonsVertexInstances( VertexInstancesToChange );
+
+	UE_LOG( LogEditableMesh, Verbose, TEXT( "* ReplaceVertexInstanceInPolygons returned" ) );
+}
+
+
+void UEditableMesh::GetVertexInstanceConnectedPolygonsInSameGroup( const FVertexInstanceID VertexInstanceID, const FPolygonID PolygonID, TArray<FPolygonID>& OutPolygonIDs ) const
+{
+	// The aim here is to determine which polygons connected to this vertex instance form part of the same soft edged group as the specified polygon.
+
+	OutPolygonIDs.Reset();
+
+	const FMeshVertexInstance& VertexInstance = VertexInstances[ VertexInstanceID.GetValue() ];
+	const FVertexID VertexID = VertexInstance.VertexID;
+
+	// Cache a list of all soft edges which share this vertex.
+	// We're only interested in finding adjacent polygons which are not the other side of a hard edge.
+	static TArray<FEdgeID> ConnectedSoftEdges;
+	GetConnectedSoftEdges( VertexID, ConnectedSoftEdges );
+
+	// Iterate through adjacent polygons which contain the given vertex, but don't cross a hard edge.
+	// Maintain a list of polygon IDs to be examined. Adjacents are added to the list if suitable.
+	// Add the start poly here.
+	static TArray<FPolygonID> PolygonsToCheck;
+	PolygonsToCheck.Reset();
+	PolygonsToCheck.Add( PolygonID );
+
+	int32 Index = 0;
+	while( Index < PolygonsToCheck.Num() )
+	{
+		const FPolygonID PolygonToCheck = PolygonsToCheck[ Index ];
+		Index++;
+
+		if( VertexInstance.ConnectedPolygons.Contains( PolygonToCheck ) )
+		{
+			OutPolygonIDs.Add( PolygonToCheck );
+
+			// Now look at its adjacent polygons. If they are joined by a soft edge which includes the vertex we're interested in, we want to consider them.
+			// We take a shortcut by doing this process in reverse: we already know all the soft edges we are interested in, so check if any of them
+			// have the current polygon as an adjacent.
+			for( const FEdgeID ConnectedSoftEdge : ConnectedSoftEdges )
+			{
+				const FMeshEdge& Edge = Edges[ ConnectedSoftEdge.GetValue() ];
+				if( Edge.ConnectedPolygons.Contains( PolygonToCheck ) )
+				{
+					for( const FPolygonID AdjacentPolygon : Edge.ConnectedPolygons )
+					{
+						// Only add new polygons which haven't yet been added to the list. This prevents circular runs of polygons triggering infinite loops.
+						PolygonsToCheck.AddUnique( AdjacentPolygon );
+					}
+				}
+			}
+		}
+	}
+}
+
+
+void UEditableMesh::SetEdgeHardness( const FEdgeID EdgeID, const bool bIsHard, const bool bIsUndo )
+{
+	checkSlow( Edges.IsAllocated( EdgeID.GetValue() ) );
+	FMeshEdge& Edge = Edges[ EdgeID.GetValue() ];
+
+	if( Edge.bIsHardEdge != bIsHard )
+	{
+		Edge.bIsHardEdge = bIsHard;
+
+		if( !bIsUndo )
+		{
+			UE_LOG( LogEditableMesh, Verbose, TEXT( "SetEdgeHardness (changed): %s %s" ), *EdgeID.ToString(), *LogHelpers::BoolToString( bIsHard ) );
+
+			if( bIsHard )
+			{
+				// Setting a hard edge potentially requires vertex instances to be split
+
+				for( const FVertexID VertexID : Edge.VertexIDs )
+				{
+					const FMeshVertex& Vertex = Vertices[ VertexID.GetValue() ];
+
+					// For each end of the edge, go through all vertex instances, determining if any need to be split.
+					// Take a copy of the vertex instances because splitting them will mutate the list we are iterating.
+					static TArray<FVertexInstanceID> VertexInstanceIDsCopy;
+					VertexInstanceIDsCopy = Vertex.VertexInstanceIDs;
+
+					// Iterate all the vertex instances for each endpoint.
+					// We are looking for vertex instances which need splitting into two different instances
+					// (because the hard edge will require separate normals for each of its polygon smoothing groups).
+					for( const FVertexInstanceID VertexInstanceID : VertexInstanceIDsCopy )
+					{
+						const FMeshVertexInstance& VertexInstance = VertexInstances[ VertexInstanceID.GetValue() ];
+
+						// Take a copy of the array of polygons connected to this instance.
+						static TArray<FPolygonID> PolygonsSharingThisVertex;
+						PolygonsSharingThisVertex = VertexInstance.ConnectedPolygons;
+						bool bFirstTime = true;
+						while( PolygonsSharingThisVertex.Num() > 0 )
+						{
+							// For the next polygon in the array, determine all other polygons in the same smoothing group which share this vertex instance.
+							static TArray<FPolygonID> PolygonsInSameSmoothingGroup;
+							GetVertexInstanceConnectedPolygonsInSameGroup( VertexInstanceID, PolygonsSharingThisVertex[ 0 ], PolygonsInSameSmoothingGroup );
+
+							// Check that all polygons in the smoothing group are attached to this vertex instance, and remove them from the master list of polygons
+							// connected to this instance. If a polygon in the smoothing group is not attached to this vertex instance, it's because it's the other
+							// side of a UV seam and hence has a distinct vertex instance.
+							for( const FPolygonID PolygonInSmoothingGroup : PolygonsInSameSmoothingGroup )
+							{
+								verify( PolygonsSharingThisVertex.Remove( PolygonInSmoothingGroup ) == 1 );
+							}
+
+							// First group which we extract: do nothing - they can keep their existing instance ID.
+							// Subsequent times round the loop, we create a new vertex instance copied from the original one, and replace connected polygon vertices with it.
+							if( !bFirstTime )
+							{
+								SplitVertexInstanceInPolygons( VertexInstanceID, PolygonsInSameSmoothingGroup );
+							}
+
+							bFirstTime = false;
+						}
+					}
+				}
+			}
+			else
+			{
+				// Setting a soft edge potentially requires vertex instances to be merged
+
+				static TArray<FVertexInstanceID> VertexInstancesToDelete;
+				VertexInstancesToDelete.Reset();
+
+				for( const FVertexID VertexID : Edge.VertexIDs )
+				{
+					// For each end of the edge which has been softened...
+					const FMeshVertex& Vertex = Vertices[ VertexID.GetValue() ];
+
+					// ...iterate through pairs of vertex instances, looking for potential to merge them
+					for( int32 IndexA = 0; IndexA < Vertex.VertexInstanceIDs.Num() - 1; ++IndexA )
+					{
+						const FVertexInstanceID VertexInstanceIDA = Vertex.VertexInstanceIDs[ IndexA ];
+						const FMeshVertexInstance& VertexInstanceA = VertexInstances[ VertexInstanceIDA.GetValue() ];
+
+						// If vertex instance isn't connected to any polygon, we can't deduce anything about its smoothing group, so skip to the next one
+						if( VertexInstanceA.ConnectedPolygons.Num() == 0 )
+						{
+							continue;
+						}
+
+						// Get mergeable set of attributes for first vertex instance
+						static TArray<FMeshElementAttributeData> AttributesA;
+						GetMergedNamedVertexInstanceAttributeData( GetMergeableVertexInstanceAttributes(), Vertex.VertexInstanceIDs[ IndexA ], TArray<FMeshElementAttributeData>(), AttributesA );
+
+						for( int32 IndexB = IndexA + 1; IndexB < Vertex.VertexInstanceIDs.Num(); ++IndexB )
+						{
+							const FVertexInstanceID VertexInstanceIDB = Vertex.VertexInstanceIDs[ IndexB ];
+
+							// If this vertex instance has been marked for deletion in a previous iteration, skip it
+							if( VertexInstancesToDelete.Contains( VertexInstanceIDB ) )
+							{
+								continue;
+							}
+
+							// Get mergeable set of attributes for second vertex instance
+							static TArray<FMeshElementAttributeData> AttributesB;
+							GetMergedNamedVertexInstanceAttributeData( GetMergeableVertexInstanceAttributes(), Vertex.VertexInstanceIDs[ IndexB ], TArray<FMeshElementAttributeData>(), AttributesB );
+
+							// If both vertex instances have equal mergeable attributes, they can potentially be merged
+							if( AreAttributeListsEqual( AttributesA, AttributesB ) )
+							{
+								const FMeshVertexInstance& VertexInstanceB = VertexInstances[ VertexInstanceIDB.GetValue() ];
+
+								// Check to see if they are in the same smoothing group, using an arbitary connected polygon in the first vertex instance as the pivot point.
+								static TArray<FVertexInstanceID> PossibleVertexInstances;
+								GetVertexInstancesInSameSoftEdgedGroup( VertexInstanceA.VertexID, VertexInstanceA.ConnectedPolygons[ 0 ], false, PossibleVertexInstances );
+
+								// If the smoothing group contains both vertex instances, they can be merged.
+								if( PossibleVertexInstances.Contains( VertexInstanceIDA ) && PossibleVertexInstances.Contains( VertexInstanceIDB ) )
+								{
+									// Change occurrences of VertexInstanceB for VertexInstanceA in VertexInstanceB's connected polygons.
+									// Note, this will cause VertexInstanceA's connected polygons list to be added to (at the end).
+									// This works because we are evaluating the number of connected polygons each time round the loop.
+									ReplaceVertexInstanceInPolygons( VertexInstanceIDB, VertexInstanceIDA, VertexInstanceB.ConnectedPolygons );
+
+									// This will also cause VertexInstanceB to be disconnected from all polygons.
+									// We mark the vertex instance for deletion here, but do not delete it until the end, as to do so would interrupt iterating through vertex instances.
+									check( VertexInstanceB.ConnectedPolygons.Num() == 0 );
+									VertexInstancesToDelete.Add( VertexInstanceIDB );
+								}
+							}
+						}
+					}
+				}
+
+				// Delete orphaned vertex instances
+				const bool bDeleteOrphanedVertices = false;
+				DeleteVertexInstances( VertexInstancesToDelete, bDeleteOrphanedVertices );
+			}
+
+			UE_LOG( LogEditableMesh, Verbose, TEXT( "* SetEdgeHardness returned" ) );
+		}
+
+		// Mark polygons adjacent to this edge as candidates for normal/tangent recalculation
+		PolygonsPendingNewTangentBasis.Append( Edge.ConnectedPolygons );
+	}
+}
+
+
+int32 UEditableMesh::GetEdgeCount() const
+{
+	return Edges.Num();
+}
+
+
+int32 UEditableMesh::GetEdgeArraySize() const
+{
+	return Edges.GetMaxIndex();
+}
+
+
+bool UEditableMesh::IsValidEdge( const FEdgeID EdgeID ) const
+{
+	return EdgeID.GetValue() >= 0 && EdgeID.GetValue() < Edges.GetMaxIndex() && Edges.IsAllocated( EdgeID.GetValue() );
+}
+
+
+FVertexID UEditableMesh::GetEdgeVertex( const FEdgeID EdgeID, const int32 EdgeVertexNumber ) const
+{
+	checkSlow( EdgeVertexNumber >= 0 && EdgeVertexNumber < 2 );
+	checkSlow( Edges.IsAllocated( EdgeID.GetValue() ) );
+	return Edges[ EdgeID.GetValue() ].VertexIDs[ EdgeVertexNumber ];
+}
+
+
+int32 UEditableMesh::GetEdgeConnectedPolygonCount( const FEdgeID EdgeID ) const
+{
+	checkSlow( Edges.IsAllocated( EdgeID.GetValue() ) );
+	const FMeshEdge& Edge = Edges[ EdgeID.GetValue() ];
+	return Edge.ConnectedPolygons.Num();
+}
+
+
+FPolygonID UEditableMesh::GetEdgeConnectedPolygon( const FEdgeID EdgeID, const int32 ConnectedPolygonNumber ) const
+{
+	checkSlow( Edges.IsAllocated( EdgeID.GetValue() ) );
+	const FMeshEdge& Edge = Edges[ EdgeID.GetValue() ];
+	const FPolygonID PolygonID = Edge.ConnectedPolygons[ ConnectedPolygonNumber ];
+	return PolygonID;
+}
+
+
+int32 UEditableMesh::GetPolygonGroupCount() const
+{
+	return PolygonGroups.Num();
+}
+
+
+int32 UEditableMesh::GetPolygonGroupArraySize() const
+{
+	return PolygonGroups.GetMaxIndex();
+}
+
+
+bool UEditableMesh::IsValidPolygonGroup( const FPolygonGroupID PolygonGroupID ) const
+{
+	return
+		PolygonGroupID.GetValue() >= 0 &&
+		PolygonGroupID.GetValue() < PolygonGroups.GetMaxIndex() &&
+		PolygonGroups.IsAllocated( PolygonGroupID.GetValue() );
+}
+
+
+int32 UEditableMesh::GetPolygonCountInGroup( const FPolygonGroupID PolygonGroupID ) const
+{
+	checkSlow( PolygonGroups.IsAllocated( PolygonGroupID.GetValue() ) );
+	return PolygonGroups[ PolygonGroupID.GetValue() ].Polygons.Num();
+}
+
+
+FPolygonID UEditableMesh::GetPolygonInGroup( const FPolygonGroupID PolygonGroupID, const int32 PolygonNumber ) const
+{
+	checkSlow( PolygonGroups.IsAllocated( PolygonGroupID.GetValue() ) );
+	const FMeshPolygonGroup& PolygonGroup = PolygonGroups[ PolygonGroupID.GetValue() ];
+	return FPolygonID( PolygonGroup.Polygons[ PolygonNumber ] );
+}
+
+
+int32 UEditableMesh::GetPolygonCount() const
+{
+	return Polygons.Num();
+}
+
+
+int32 UEditableMesh::GetPolygonArraySize() const
+{
+	return Polygons.GetMaxIndex();
+}
+
+
+bool UEditableMesh::IsValidPolygon( const FPolygonID PolygonID ) const
+{
+	return
+		PolygonID.GetValue() >= 0 &&
+		PolygonID.GetValue() < Polygons.GetMaxIndex() &&
+		Polygons.IsAllocated( PolygonID.GetValue() );
+}
+
+
+FPolygonGroupID UEditableMesh::GetGroupForPolygon( const FPolygonID PolygonID ) const
+{
+	checkSlow( Polygons.IsAllocated( PolygonID.GetValue() ) );
+	return Polygons[ PolygonID.GetValue() ].PolygonGroupID;
+}
+
+
+int32 UEditableMesh::GetPolygonPerimeterVertexCount( const FPolygonID PolygonID ) const
+{
+	checkSlow( Polygons.IsAllocated( PolygonID.GetValue() ) );
+	const FMeshPolygon& Polygon = Polygons[ PolygonID.GetValue() ];
+
+	return Polygon.PerimeterContour.VertexInstanceIDs.Num();
+}
+
+
+FVertexID UEditableMesh::GetPolygonPerimeterVertex( const FPolygonID PolygonID, const int32 PolygonVertexNumber ) const
+{
+	checkSlow( Polygons.IsAllocated( PolygonID.GetValue() ) );
+	const FMeshPolygon& Polygon = Polygons[ PolygonID.GetValue() ];
+	const FVertexInstanceID VertexInstanceID = Polygon.PerimeterContour.VertexInstanceIDs[ PolygonVertexNumber ];
+	const FMeshVertexInstance& VertexInstance = VertexInstances[ VertexInstanceID.GetValue() ];
+	return VertexInstance.VertexID;
+}
+
+
+FVertexInstanceID UEditableMesh::GetPolygonPerimeterVertexInstance( const FPolygonID PolygonID, const int32 PolygonVertexNumber ) const
+{
+	checkSlow( Polygons.IsAllocated( PolygonID.GetValue() ) );
+	const FMeshPolygon& Polygon = Polygons[ PolygonID.GetValue() ];
+	return Polygon.PerimeterContour.VertexInstanceIDs[ PolygonVertexNumber ];
+}
+
+
+FVector4 UEditableMesh::GetPolygonPerimeterVertexAttribute( const FPolygonID PolygonID, const int32 PolygonVertexNumber, const FName AttributeName, const int32 AttributeIndex ) const
+{
+	checkSlow( Polygons.IsAllocated( PolygonID.GetValue() ) );
+	const FMeshPolygon& Polygon = Polygons[ PolygonID.GetValue() ];
+	const FVertexInstanceID VertexInstanceID = Polygon.PerimeterContour.VertexInstanceIDs[ PolygonVertexNumber ];
+	return GetVertexInstanceAttribute( VertexInstanceID, AttributeName, AttributeIndex );
+}
+
+
+int32 UEditableMesh::GetPolygonHoleCount( const FPolygonID PolygonID ) const
+{
+	checkSlow( Polygons.IsAllocated( PolygonID.GetValue() ) );
+	const FMeshPolygon& Polygon = Polygons[ PolygonID.GetValue() ];
+
+	return Polygon.HoleContours.Num();
+}
+
+
+int32 UEditableMesh::GetPolygonHoleVertexCount( const FPolygonID PolygonID, const int32 HoleNumber ) const
+{
+	checkSlow( Polygons.IsAllocated( PolygonID.GetValue() ) );
+	const FMeshPolygon& Polygon = Polygons[ PolygonID.GetValue() ];
+	return Polygon.HoleContours[ HoleNumber ].VertexInstanceIDs.Num();
+}
+
+
+FVertexID UEditableMesh::GetPolygonHoleVertex( const FPolygonID PolygonID, const int32 HoleNumber, const int32 PolygonVertexNumber ) const
+{
+	checkSlow( Polygons.IsAllocated( PolygonID.GetValue() ) );
+	const FMeshPolygon& Polygon = Polygons[ PolygonID.GetValue() ];
+	const FVertexInstanceID VertexInstanceID = Polygon.HoleContours[ HoleNumber ].VertexInstanceIDs[ PolygonVertexNumber ];
+	const FMeshVertexInstance& VertexInstance = VertexInstances[ VertexInstanceID.GetValue() ];
+	return VertexInstance.VertexID;
+}
+
+
+FVertexInstanceID UEditableMesh::GetPolygonHoleVertexInstance( const FPolygonID PolygonID, const int32 HoleNumber, const int32 PolygonVertexNumber ) const
+{
+	checkSlow( Polygons.IsAllocated( PolygonID.GetValue() ) );
+	const FMeshPolygon& Polygon = Polygons[ PolygonID.GetValue() ];
+	return Polygon.HoleContours[ HoleNumber ].VertexInstanceIDs[ PolygonVertexNumber ];
+}
+
+
+FVector4 UEditableMesh::GetPolygonHoleVertexAttribute( const FPolygonID PolygonID, const int32 HoleNumber, const int32 PolygonVertexNumber, const FName AttributeName, const int32 AttributeIndex ) const
+{
+	checkSlow( Polygons.IsAllocated( PolygonID.GetValue() ) );
+	const FMeshPolygon& Polygon = Polygons[ PolygonID.GetValue() ];
+	const FVertexInstanceID VertexInstanceID = Polygon.HoleContours[ HoleNumber ].VertexInstanceIDs[ PolygonVertexNumber ];
+	return GetVertexInstanceAttribute( VertexInstanceID, AttributeName, AttributeIndex );
+}
+
+
 const TArray<FName>& UEditableMesh::GetValidVertexAttributes()
 {
 	static TArray<FName> ValidVertexAttributes;
@@ -110,18 +1333,18 @@
 }
 
 
-const TArray<FName>& UEditableMesh::GetValidPolygonVertexAttributes()
+const TArray<FName>& UEditableMesh::GetValidVertexInstanceAttributes()
 {
-	static TArray<FName> ValidPolygonVertexAttributes;
-	if( ValidPolygonVertexAttributes.Num() == 0 )
+	static TArray<FName> ValidVertexInstanceAttributes;
+	if( ValidVertexInstanceAttributes.Num() == 0 )
 	{
-		ValidPolygonVertexAttributes.Add( UEditableMeshAttribute::VertexNormal() );
-		ValidPolygonVertexAttributes.Add( UEditableMeshAttribute::VertexTangent() );
-		ValidPolygonVertexAttributes.Add( UEditableMeshAttribute::VertexBinormalSign() );
-		ValidPolygonVertexAttributes.Add( UEditableMeshAttribute::VertexTextureCoordinate() );
-		ValidPolygonVertexAttributes.Add( UEditableMeshAttribute::VertexColor() );
+		ValidVertexInstanceAttributes.Add( UEditableMeshAttribute::VertexNormal() );
+		ValidVertexInstanceAttributes.Add( UEditableMeshAttribute::VertexTangent() );
+		ValidVertexInstanceAttributes.Add( UEditableMeshAttribute::VertexBinormalSign() );
+		ValidVertexInstanceAttributes.Add( UEditableMeshAttribute::VertexTextureCoordinate() );
+		ValidVertexInstanceAttributes.Add( UEditableMeshAttribute::VertexColor() );
 	}
-	return ValidPolygonVertexAttributes;
+	return ValidVertexInstanceAttributes;
 }
 
 
@@ -137,6 +1360,18 @@
 }
 
 
+const TArray<FName>& UEditableMesh::GetValidPolygonAttributes()
+{
+	static TArray<FName> ValidPolygonAttributes;
+	if( ValidPolygonAttributes.Num() == 0 )
+	{
+		ValidPolygonAttributes.Add( UEditableMeshAttribute::PolygonNormal() );
+		ValidPolygonAttributes.Add( UEditableMeshAttribute::PolygonCenter() );
+	}
+	return ValidPolygonAttributes;
+}
+
+
 int32 UEditableMesh::GetMaxAttributeIndex( const FName AttributeName ) const
 {
 	if( AttributeName == UEditableMeshAttribute::VertexTextureCoordinate() )
@@ -148,44 +1383,19 @@
 }
 
 
-FSectionID UEditableMesh::GetFirstValidSection() const
+FPolygonGroupID UEditableMesh::GetFirstValidPolygonGroup() const
 {
-	FSectionID FirstValidSectionID = FSectionID::Invalid;
+	FPolygonGroupID FirstValidPolygonGroupID = FPolygonGroupID::Invalid;
 
-	// @todo mesheditor perf: We can do this without iterating if we make the function virtual and allow the backend to return the first valid section
-	const int32 SectionArraySize = GetSectionArraySize();
-	for( int32 SectionIndex = 0; SectionIndex < SectionArraySize; ++SectionIndex )
+	if( PolygonGroups.Num() > 0 )
 	{
-		const FSectionID SectionID( SectionIndex );
-		if( IsValidSection( SectionID ) )
-		{
-			FirstValidSectionID = SectionID;
-		}
+		FirstValidPolygonGroupID = FPolygonGroupID( PolygonGroups.CreateConstIterator().GetIndex() );
 	}
 
-	return FirstValidSectionID;
+	return FirstValidPolygonGroupID;
 }
 
 
-int32 UEditableMesh::GetTotalPolygonCount() const
-{
-	int32 TotalPolygonCount = 0;
-
-	const int32 SectionArraySize = GetSectionArraySize();
-	for( int32 SectionIndex = 0; SectionIndex < SectionArraySize; ++SectionIndex )
-	{
-		const FSectionID SectionID( SectionIndex );
-		if( IsValidSection( SectionID ) )
-		{
-			const int32 PolygonCount = GetPolygonCount( SectionID );
-			TotalPolygonCount += PolygonCount;
-		}
-	}
-
-	return TotalPolygonCount;
-}
-
-
 int32 UEditableMesh::GetTextureCoordinateCount() const
 {
 	return this->TextureCoordinateCount;
@@ -218,63 +1428,66 @@
 }
 
 
-void UEditableMesh::GetVertexConnectedPolygons( const FVertexID VertexID, TArray<FPolygonRef>& OutConnectedPolygonRefs ) const
+void UEditableMesh::GetVertexConnectedPolygons( const FVertexID VertexID, TArray<FPolygonID>& OutConnectedPolygonIDs ) const
 {
-	OutConnectedPolygonRefs.Reset();
+	OutConnectedPolygonIDs.Reset();
 
-	const int32 ConnectedEdgeCount = GetVertexConnectedEdgeCount( VertexID );
-	for( int32 EdgeNumber = 0; EdgeNumber < ConnectedEdgeCount; ++EdgeNumber )
+	checkSlow( Vertices.IsAllocated( VertexID.GetValue() ) );
+	const FMeshVertex& Vertex = Vertices[ VertexID.GetValue() ];
+	for( const FVertexInstanceID VertexInstanceID : Vertex.VertexInstanceIDs )
 	{
-		const FEdgeID ConnectedEdgeID = GetVertexConnectedEdge( VertexID, EdgeNumber );
-
-		const int32 ConnectedPolygonCount = GetEdgeConnectedPolygonCount( ConnectedEdgeID );
-		for( int32 PolygonNumber = 0; PolygonNumber < ConnectedPolygonCount; ++PolygonNumber )
-		{
-			const FPolygonRef PolygonRef = GetEdgeConnectedPolygon( ConnectedEdgeID, PolygonNumber );
-
-			// Add uniquely, because a vertex's connected edges may neighbor the same polygon
-			OutConnectedPolygonRefs.AddUnique( PolygonRef );
-		}
+		checkSlow( VertexInstances.IsAllocated( VertexInstanceID.GetValue() ) );
+		const FMeshVertexInstance& VertexInstance = VertexInstances[ VertexInstanceID.GetValue() ];
+		OutConnectedPolygonIDs.Append( VertexInstance.ConnectedPolygons );
 	}
 }
 
 
-void UEditableMesh::GetVertexAdjacentVertices( const FVertexID VertexID, TArray< FVertexID >& OutAdjacentVertexIDs ) const
+void UEditableMesh::GetVertexInstanceConnectedPolygons( const FVertexInstanceID VertexInstanceID, TArray<FPolygonID>& OutConnectedPolygonIDs ) const
 {
+	checkSlow( VertexInstances.IsAllocated( VertexInstanceID.GetValue() ) );
+	const FMeshVertexInstance& VertexInstance = VertexInstances[ VertexInstanceID.GetValue() ];
+
+	OutConnectedPolygonIDs = VertexInstance.ConnectedPolygons;
+}
+
+
+int32 UEditableMesh::GetVertexInstanceConnectedPolygonCount( const FVertexInstanceID VertexInstanceID ) const
+{
+	checkSlow( VertexInstances.IsAllocated( VertexInstanceID.GetValue() ) );
+	const FMeshVertexInstance& VertexInstance = VertexInstances[ VertexInstanceID.GetValue() ];
+	return VertexInstance.ConnectedPolygons.Num();
+}
+
+
+void UEditableMesh::GetVertexAdjacentVertices( const FVertexID VertexID, TArray<FVertexID>& OutAdjacentVertexIDs ) const
+{
 	OutAdjacentVertexIDs.Reset();
 
-	const int32 ConnectedEdgeCount = GetVertexConnectedEdgeCount( VertexID );
-	for( int32 EdgeNumber = 0; EdgeNumber < ConnectedEdgeCount; ++EdgeNumber )
+	const FMeshVertex& Vertex = Vertices[ VertexID.GetValue() ];
+	for( const FEdgeID EdgeID : Vertex.ConnectedEdgeIDs )
 	{
-		const FEdgeID ConnectedEdgeID = GetVertexConnectedEdge( VertexID, EdgeNumber );
-		
-		FVertexID EdgeVertexIDs[ 2 ];
-		GetEdgeVertices( ConnectedEdgeID, /* Out */ EdgeVertexIDs[0], /* Out */ EdgeVertexIDs[1] );
-
-		OutAdjacentVertexIDs.Add( EdgeVertexIDs[ 0 ] == VertexID ? EdgeVertexIDs[ 1 ] : EdgeVertexIDs[ 0 ] );
+		checkSlow( Edges.IsAllocated( EdgeID.GetValue() ) );
+		const FMeshEdge& Edge = Edges[ EdgeID.GetValue() ];
+		OutAdjacentVertexIDs.Add( Edge.VertexIDs[ 0 ] == VertexID ? Edge.VertexIDs[ 1 ] : Edge.VertexIDs[ 0 ] );
 	}
 }
 
 
 void UEditableMesh::GetEdgeVertices( const FEdgeID EdgeID, FVertexID& OutEdgeVertexID0, FVertexID& OutEdgeVertexID1 ) const
 {
-	OutEdgeVertexID0 = GetEdgeVertex( EdgeID, 0 );
-	OutEdgeVertexID1 = GetEdgeVertex( EdgeID, 1 );
+	checkSlow( Edges.IsAllocated( EdgeID.GetValue() ) );
+	const FMeshEdge& Edge = Edges[ EdgeID.GetValue() ];
+	OutEdgeVertexID0 = Edge.VertexIDs[ 0 ];
+	OutEdgeVertexID1 = Edge.VertexIDs[ 1 ];
 }
 
 
-void UEditableMesh::GetEdgeConnectedPolygons( const FEdgeID EdgeID, TArray<FPolygonRef>& OutConnectedPolygonRefs ) const
+void UEditableMesh::GetEdgeConnectedPolygons( const FEdgeID EdgeID, TArray<FPolygonID>& OutConnectedPolygonIDs ) const
 {
-	OutConnectedPolygonRefs.Reset();
-
-	const int32 ConnectedPolygonCount = GetEdgeConnectedPolygonCount( EdgeID );
-	for( int32 PolygonNumber = 0; PolygonNumber < ConnectedPolygonCount; ++PolygonNumber )
-	{
-		const FPolygonRef PolygonRef = GetEdgeConnectedPolygon( EdgeID, PolygonNumber );
-
-		// Add uniquely, because a vertex's connected edges may neighbor the same polygon
-		OutConnectedPolygonRefs.AddUnique( PolygonRef );
-	}
+	checkSlow( Edges.IsAllocated( EdgeID.GetValue() ) );
+	const FMeshEdge& Edge = Edges[ EdgeID.GetValue() ];
+	OutConnectedPolygonIDs = Edge.ConnectedPolygons;
 }
 
 
@@ -325,16 +1538,16 @@
 		const FVertexID FromVertexID = CurrentItem.Get<1>();
 		EdgeIDs.Add( CurrentEdgeID );
 
-		// Get the polygons connected to this edge. When continuing the loop, the criterion is that new edges must share no polygons with this edge,
+		// Get the edge so we can see which polygons are connected to it.
+		// When continuing the loop, the criterion is that new edges must share no polygons with this edge,
 		// i.e. they are the other side of a perpendicular edge.
-		static TArray<FPolygonRef> ConnectedPolygons;
-		GetEdgeConnectedPolygons( CurrentEdgeID, ConnectedPolygons );
+		checkSlow( Edges.IsAllocated( CurrentEdgeID.GetValue() ) );
+		const FMeshEdge& Edge = Edges[ CurrentEdgeID.GetValue() ];
 
 		// Now look for edges connected to each end of this edge
-		for( int32 ConnectedVertexIndex = 0; ConnectedVertexIndex < 2; ++ConnectedVertexIndex )
+		for( const FVertexID ConnectedVertexID : Edge.VertexIDs )
 		{
 			// If we have already processed this vertex, skip it
-			const FVertexID ConnectedVertexID = GetEdgeVertex( CurrentEdgeID, ConnectedVertexIndex );
 			if( ConnectedVertexID == FromVertexID )
 			{
 				continue;
@@ -344,21 +1557,19 @@
 			FEdgeID AdjacentEdgeID = FEdgeID::Invalid;
 
 			// Iterate through all edges connected to this vertex
-			const int32 VertexConnectedEdgeCount = GetVertexConnectedEdgeCount( ConnectedVertexID );
-			for( int32 EdgeIndex = 0; EdgeIndex < VertexConnectedEdgeCount; ++EdgeIndex )
+			const FMeshVertex& Vertex = Vertices[ ConnectedVertexID.GetValue() ];
+			for( const FEdgeID ConnectedEdgeID : Vertex.ConnectedEdgeIDs )
 			{
-				const FEdgeID ConnectedEdgeID = GetVertexConnectedEdge( ConnectedVertexID, EdgeIndex );
-
 				// If this edge hasn't been added to the loop...
 				if( !EdgeIDs.Contains( ConnectedEdgeID ) )
 				{
 					// ...see if it shares any polygons with the original edge (intersection operation)
-					const int32 EdgeConnectedPolygonsCount = GetEdgeConnectedPolygonCount( ConnectedEdgeID );
 					bool bIsCandidateEdge = true;
-					for( int32 PolygonIndex = 0; PolygonIndex < EdgeConnectedPolygonsCount; ++PolygonIndex )
+					checkSlow( Edges.IsAllocated( ConnectedEdgeID.GetValue() ) );
+					const FMeshEdge& ConnectedEdge = Edges[ ConnectedEdgeID.GetValue() ];
+					for( const FPolygonID ConnectedPolygonID : ConnectedEdge.ConnectedPolygons )
 					{
-						const FPolygonRef ConnectedPolygonID = GetEdgeConnectedPolygon( ConnectedEdgeID, PolygonIndex );
-						if( ConnectedPolygons.Contains( ConnectedPolygonID ) )
+						if( Edge.ConnectedPolygons.Contains( ConnectedPolygonID ) )
 						{
 							bIsCandidateEdge = false;
 							break;
@@ -397,268 +1608,218 @@
 }
 
 
-int32 UEditableMesh::GetPolygonPerimeterEdgeCount( const FPolygonRef PolygonRef ) const
+int32 UEditableMesh::GetPolygonPerimeterEdgeCount( const FPolygonID PolygonID ) const
 {
 	// All polygons have the same number of edges as they do vertices
-	return GetPolygonPerimeterVertexCount( PolygonRef );
+	return GetPolygonPerimeterVertexCount( PolygonID );
 }
 
 
-int32 UEditableMesh::GetPolygonHoleEdgeCount( const FPolygonRef PolygonRef, const int32 HoleNumber ) const
+int32 UEditableMesh::GetPolygonHoleEdgeCount( const FPolygonID PolygonID, const int32 HoleNumber ) const
 {
 	// All polygons have the same number of edges as they do vertices
-	return GetPolygonHoleVertexCount( PolygonRef, HoleNumber );
+	return GetPolygonHoleVertexCount( PolygonID, HoleNumber );
 }
 
 
-void UEditableMesh::GetPolygonPerimeterVertices( const FPolygonRef PolygonRef, TArray<FVertexID>& OutPolygonPerimeterVertexIDs ) const
+void UEditableMesh::GetPolygonPerimeterVertices( const FPolygonID PolygonID, TArray<FVertexID>& OutPolygonPerimeterVertexIDs ) const
 {
-	const int32 NumPolygonPerimeterVertices = GetPolygonPerimeterVertexCount( PolygonRef );
+	checkSlow( Polygons.IsAllocated( PolygonID.GetValue() ) );
+	const FMeshPolygon& Polygon = Polygons[ PolygonID.GetValue() ];
+	OutPolygonPerimeterVertexIDs.SetNumUninitialized( Polygon.PerimeterContour.VertexInstanceIDs.Num(), false );
 
-	OutPolygonPerimeterVertexIDs.SetNumUninitialized( NumPolygonPerimeterVertices, false );
-	for( int32 VertexNumber = 0; VertexNumber < NumPolygonPerimeterVertices; ++VertexNumber )
+	int32 Index = 0;
+	for( const FVertexInstanceID VertexInstanceID : Polygon.PerimeterContour.VertexInstanceIDs )
 	{
-		OutPolygonPerimeterVertexIDs[ VertexNumber ] = GetPolygonPerimeterVertex( PolygonRef, VertexNumber );
+		const FMeshVertexInstance& VertexInstance = VertexInstances[ VertexInstanceID.GetValue() ];
+		OutPolygonPerimeterVertexIDs[ Index ] = VertexInstance.VertexID;
+		Index++;
 	}
 }
 
 
-void UEditableMesh::GetPolygonHoleVertices( const FPolygonRef PolygonRef, const int32 HoleNumber, TArray<FVertexID>& OutHoleVertexIDs ) const
+void UEditableMesh::GetPolygonPerimeterVertexInstances( const FPolygonID PolygonID, TArray<FVertexInstanceID>& OutPolygonPerimeterVertexInstanceIDs ) const
 {
-	const int32 NumHoleVertices = GetPolygonHoleVertexCount( PolygonRef, HoleNumber );
+	checkSlow( Polygons.IsAllocated( PolygonID.GetValue() ) );
+	const FMeshPolygon& Polygon = Polygons[ PolygonID.GetValue() ];
+	OutPolygonPerimeterVertexInstanceIDs = Polygon.PerimeterContour.VertexInstanceIDs;
+}
 
-	OutHoleVertexIDs.SetNumUninitialized( NumHoleVertices, false );
-	for( int32 VertexNumber = 0; VertexNumber < NumHoleVertices; ++VertexNumber )
+
+void UEditableMesh::GetPolygonHoleVertices( const FPolygonID PolygonID, const int32 HoleNumber, TArray<FVertexID>& OutHoleVertexIDs ) const
+{
+	checkSlow( Polygons.IsAllocated( PolygonID.GetValue() ) );
+	const FMeshPolygon& Polygon = Polygons[ PolygonID.GetValue() ];
+	const FMeshPolygonContour& Contour = Polygon.HoleContours[ HoleNumber ];
+
+	OutHoleVertexIDs.SetNumUninitialized( Contour.VertexInstanceIDs.Num(), false );
+
+	int32 Index = 0;
+	for( const FVertexInstanceID VertexInstanceID : Contour.VertexInstanceIDs )
 	{
-		OutHoleVertexIDs[ VertexNumber ] = GetPolygonHoleVertex( PolygonRef, HoleNumber, VertexNumber );
+		const FMeshVertexInstance& VertexInstance = VertexInstances[ VertexInstanceID.GetValue() ];
+		OutHoleVertexIDs[ Index ] = VertexInstance.VertexID;
+		Index++;
 	}
 }
 
 
-FEdgeID UEditableMesh::GetPolygonPerimeterEdge( const FPolygonRef PolygonRef, const int32 PerimeterEdgeNumber, bool& bOutEdgeWindingIsReversedForPolygon ) const
+void UEditableMesh::GetPolygonHoleVertexInstances( const FPolygonID PolygonID, const int32 HoleNumber, TArray<FVertexInstanceID>& OutHoleVertexInstanceIDs ) const
 {
-	FEdgeID FoundEdgeID = FEdgeID::Invalid;
-	bool bFoundEdge = false;
-	bOutEdgeWindingIsReversedForPolygon = false;
+	checkSlow( Polygons.IsAllocated( PolygonID.GetValue() ) );
+	const FMeshPolygon& Polygon = Polygons[ PolygonID.GetValue() ];
+	const FMeshPolygonContour& Contour = Polygon.HoleContours[ HoleNumber ];
 
-	const int32 NumPolygonPerimeterEdges = GetPolygonPerimeterEdgeCount( PolygonRef );
-	const int32 NumPolygonPerimeterVertices = NumPolygonPerimeterEdges;		// Edge and vertex count are always the same
+	OutHoleVertexInstanceIDs = Contour.VertexInstanceIDs;
+}
 
-	check( NumPolygonPerimeterEdges > 0 );
-	for( int32 CurrentPerimeterEdgeNumber = 0; CurrentPerimeterEdgeNumber < NumPolygonPerimeterEdges; ++CurrentPerimeterEdgeNumber )
+
+FEdgeID UEditableMesh::GetPolygonContourEdge( const FMeshPolygonContour& Contour, const int32 ContourEdgeNumber, bool& bOutEdgeWindingIsReversed ) const
+{
+	const int32 NumPolygonContourEdges = Contour.VertexInstanceIDs.Num();
+
+	const FVertexInstanceID VertexInstanceID = Contour.VertexInstanceIDs[ ContourEdgeNumber ];
+	const FVertexInstanceID NextVertexInstanceID = Contour.VertexInstanceIDs[ ( ContourEdgeNumber + 1 ) % NumPolygonContourEdges ];
+
+	checkSlow( VertexInstances.IsAllocated( VertexInstanceID.GetValue() ) );
+	checkSlow( VertexInstances.IsAllocated( NextVertexInstanceID.GetValue() ) );
+	const FVertexID VertexID = VertexInstances[ VertexInstanceID.GetValue() ].VertexID;
+	const FVertexID NextVertexID = VertexInstances[ NextVertexInstanceID.GetValue() ].VertexID;
+
+	return GetVertexPairEdge( VertexID, NextVertexID, bOutEdgeWindingIsReversed );
+}
+
+
+FEdgeID UEditableMesh::GetPolygonPerimeterEdge( const FPolygonID PolygonID, const int32 PerimeterEdgeNumber, bool& bOutEdgeWindingIsReversedForPolygon ) const
+{
+	checkSlow( Polygons.IsAllocated( PolygonID.GetValue() ) );
+	const FMeshPolygon& Polygon = Polygons[ PolygonID.GetValue() ];
+	return GetPolygonContourEdge( Polygon.PerimeterContour, PerimeterEdgeNumber, bOutEdgeWindingIsReversedForPolygon );
+}
+
+
+FEdgeID UEditableMesh::GetVertexPairEdge( const FVertexID StartVertexID, const FVertexID EndVertexID, bool& bOutEdgeWindingIsReversed ) const
+{
+	bOutEdgeWindingIsReversed = false;
+
+	const FMeshVertex& StartVertex = Vertices[ StartVertexID.GetValue() ];
+	for( const FEdgeID VertexConnectedEdgeID : StartVertex.ConnectedEdgeIDs )
 	{
-		const int32 PerimeterVertexNumber = CurrentPerimeterEdgeNumber;	// Edge and vertex counts are always the same
+		const FMeshEdge& Edge = Edges[ VertexConnectedEdgeID.GetValue() ];
 
-		const FVertexID VertexID = GetPolygonPerimeterVertex( PolygonRef, PerimeterVertexNumber );
-		const FVertexID NextVertexID = GetPolygonPerimeterVertex( PolygonRef, ( PerimeterVertexNumber + 1 ) % NumPolygonPerimeterVertices );
+		// Try the edge's first vertex.  Does it point to us?
+		FVertexID OtherEdgeVertexID = Edge.VertexIDs[ 0 ];
+		bOutEdgeWindingIsReversed = false;
+		if( OtherEdgeVertexID == StartVertexID )
+		{
+			// Must be the the other one
+			OtherEdgeVertexID = Edge.VertexIDs[ 1 ];
+			bOutEdgeWindingIsReversed = true;
+		}
 
-		// Find the edge that connects these vertices
-		bool bFoundCurrentEdge = false;
-
-		const int32 NumVertexConnectedEdges = GetVertexConnectedEdgeCount( VertexID );	// @todo mesheditor perf: Could be made faster by always iterating over the vertex with the fewest number of connected edges, instead of the first edge vertex
-		for( int32 VertexEdgeNumber = 0; VertexEdgeNumber < NumVertexConnectedEdges; ++VertexEdgeNumber )
+		if( OtherEdgeVertexID == EndVertexID )
 		{
-			const FEdgeID VertexConnectedEdgeID = GetVertexConnectedEdge( VertexID, VertexEdgeNumber );
-
-			// Try the edge's first vertex.  Does it point to us?
-			FVertexID OtherEdgeVertexID = GetEdgeVertex( VertexConnectedEdgeID, 0 );
-			bOutEdgeWindingIsReversedForPolygon = false;
-			if( OtherEdgeVertexID == VertexID )
-			{
-				// Must be the the other one
-				OtherEdgeVertexID = GetEdgeVertex( VertexConnectedEdgeID, 1 );
-				bOutEdgeWindingIsReversedForPolygon = true;
-			}
-
-			if( OtherEdgeVertexID == NextVertexID )
-			{
-				// We found the edge!
-				FoundEdgeID = VertexConnectedEdgeID;
-				bFoundCurrentEdge = true;
-				break;
-			}
-		}
-		check( bFoundCurrentEdge );
-
-		if( CurrentPerimeterEdgeNumber == PerimeterEdgeNumber )
-		{
-			bFoundEdge = true;
-			break;
+			// We found the edge!
+			return VertexConnectedEdgeID;
 		}
 	}
 
-	check( bFoundEdge );
-	return FoundEdgeID;
+	return FEdgeID::Invalid;
 }
 
 
-FEdgeID UEditableMesh::GetPolygonHoleEdge( const FPolygonRef PolygonRef, const int32 HoleNumber, const int32 HoleEdgeNumber ) const
+FEdgeID UEditableMesh::GetPolygonHoleEdge( const FPolygonID PolygonID, const int32 HoleNumber, const int32 HoleEdgeNumber ) const
 {
-	bool bFoundEdge = false;
-	FEdgeID FoundEdgeID = FEdgeID::Invalid;
+	checkSlow( Polygons.IsAllocated( PolygonID.GetValue() ) );
+	const FMeshPolygon& Polygon = Polygons[ PolygonID.GetValue() ];
+	const FMeshPolygonContour& Contour = Polygon.HoleContours[ HoleNumber ];
 
-	const int32 NumHoleEdges = GetPolygonHoleEdgeCount( PolygonRef, HoleNumber );
-	const int32 NumHoleVertices = NumHoleEdges;		// Edge and vertex count are always the same
-
-	check( NumHoleEdges > 0 );
-	for( int32 CurrentHoleEdgeNumber = 0; CurrentHoleEdgeNumber < NumHoleEdges; ++CurrentHoleEdgeNumber )
-	{
-		const int32 HoleVertexNumber = CurrentHoleEdgeNumber;	// Edge and vertex counts are always the same
-
-		const FVertexID VertexID = GetPolygonHoleVertex( PolygonRef, HoleNumber, HoleVertexNumber );
-		const FVertexID NextVertexID = GetPolygonHoleVertex( PolygonRef, HoleNumber, ( HoleVertexNumber + 1 ) % NumHoleVertices );
-
-		// Find the edge that connects these vertices
-		bool bFoundCurrentEdge = false;
-
-		const int32 NumVertexConnectedEdges = GetVertexConnectedEdgeCount( VertexID );	// @todo mesheditor perf: Could be made faster by always iterating over the vertex with the fewest number of connected edges, instead of the first edge vertex
-		for( int32 VertexEdgeNumber = 0; VertexEdgeNumber < NumVertexConnectedEdges; ++VertexEdgeNumber )
-		{
-			const FEdgeID VertexConnectedEdgeID = GetVertexConnectedEdge( VertexID, VertexEdgeNumber );
-
-			// Try the edge's first vertex.  Does it point to us?
-			FVertexID OtherEdgeVertexID = GetEdgeVertex( VertexConnectedEdgeID, 0 );
-			if( OtherEdgeVertexID == VertexID )
-			{
-				// Must be the the other one
-				OtherEdgeVertexID = GetEdgeVertex( VertexConnectedEdgeID, 1 );
-			}
-
-			if( OtherEdgeVertexID == NextVertexID )
-			{
-				// We found the edge!
-				FoundEdgeID = VertexConnectedEdgeID;
-				bFoundCurrentEdge = true;
-				break;
-			}
-		}
-		check( bFoundCurrentEdge );
-
-		if( CurrentHoleEdgeNumber == HoleEdgeNumber )
-		{
-			bFoundEdge = true;
-			break;
-		}
-	}
-
-	check( bFoundEdge );
-	return FoundEdgeID;
+	bool bOutEdgeWindingIsReversed;
+	return GetPolygonContourEdge( Contour, HoleEdgeNumber, bOutEdgeWindingIsReversed );
 }
 
 
-void UEditableMesh::GetPolygonPerimeterEdges( const FPolygonRef PolygonRef, TArray<FEdgeID>& OutPolygonPerimeterEdgeIDs ) const
+void UEditableMesh::GetPolygonContourEdges( const FMeshPolygonContour& Contour, TArray<FEdgeID>& OutPolygonContourEdgeIDs ) const
 {
-	const int32 NumPolygonPerimeterEdges = GetPolygonPerimeterEdgeCount( PolygonRef );
-	const int32 NumPolygonPerimeterVertices = NumPolygonPerimeterEdges;		// Edge and vertex count are always the same
+	const int32 NumContourEdges = Contour.VertexInstanceIDs.Num();
+	OutPolygonContourEdgeIDs.SetNumUninitialized( NumContourEdges, false );
 
-	OutPolygonPerimeterEdgeIDs.SetNumUninitialized( NumPolygonPerimeterEdges, false );
-
-	for( int32 PerimeterEdgeNumber = 0; PerimeterEdgeNumber < NumPolygonPerimeterEdges; ++PerimeterEdgeNumber )
+	for( int32 EdgeNumber = 0; EdgeNumber < NumContourEdges; ++EdgeNumber )
 	{
-		const int32 PerimeterVertexNumber = PerimeterEdgeNumber;	// Edge and vertex counts are always the same
-
-		const FVertexID VertexID = GetPolygonPerimeterVertex( PolygonRef, PerimeterVertexNumber );
-		const FVertexID NextVertexID = GetPolygonPerimeterVertex( PolygonRef, ( PerimeterVertexNumber + 1 ) % NumPolygonPerimeterVertices );
-
-		// Find the edge that connects these vertices
-		FEdgeID FoundEdgeID = FEdgeID::Invalid;
-		bool bFoundEdge = false;
-
-		const int32 NumVertexConnectedEdges = GetVertexConnectedEdgeCount( VertexID );	// @todo mesheditor perf: Could be made faster by always iterating over the vertex with the fewest number of connected edges, instead of the first edge vertex
-		for( int32 VertexEdgeNumber = 0; VertexEdgeNumber < NumVertexConnectedEdges; ++VertexEdgeNumber )
-		{
-			const FEdgeID VertexConnectedEdgeID = GetVertexConnectedEdge( VertexID, VertexEdgeNumber );
-
-			// Try the edge's first vertex.  Does it point to our next edge?
-			FVertexID OtherEdgeVertexID = GetEdgeVertex( VertexConnectedEdgeID, 0 );
-			if( OtherEdgeVertexID == VertexID )
-			{
-				// Must be the the other one
-				OtherEdgeVertexID = GetEdgeVertex( VertexConnectedEdgeID, 1 );
-			}
-			else
-			{
-				check( GetEdgeVertex( VertexConnectedEdgeID, 1 ) == VertexID );
-			}
-
-			if( OtherEdgeVertexID == NextVertexID )
-			{
-				// We found the edge!
-				FoundEdgeID = VertexConnectedEdgeID;
-				bFoundEdge = true;
-				break;
-			}
-		}
-
-		// @todo mesheditor perf: A lot of these inner-loop check() statements need to be disabled (or changes to checkSlow()) for faster performance
-		check( bFoundEdge );
-
-		OutPolygonPerimeterEdgeIDs[ PerimeterEdgeNumber ] = FoundEdgeID;
+		bool bOutEdgeWindingIsReversed;
+		OutPolygonContourEdgeIDs[ EdgeNumber ] = GetPolygonContourEdge( Contour, EdgeNumber, bOutEdgeWindingIsReversed );
 	}
 }
 
 
-void UEditableMesh::GetPolygonHoleEdges( const FPolygonRef PolygonRef, const int32 HoleNumber, TArray<FEdgeID>& OutHoleEdgeIDs ) const
+void UEditableMesh::GetPolygonPerimeterEdges( const FPolygonID PolygonID, TArray<FEdgeID>& OutPolygonPerimeterEdgeIDs ) const
 {
-	const int32 NumHoleEdges = GetPolygonHoleEdgeCount( PolygonRef, HoleNumber );
-	const int32 NumHoleVertices = NumHoleEdges;		// Edge and vertex count are always the same
-
-	OutHoleEdgeIDs.SetNumUninitialized( NumHoleEdges, false );
-
-	for( int32 HoleEdgeNumber = 0; HoleEdgeNumber < NumHoleEdges; ++HoleEdgeNumber )
-	{
-		const int32 HoleVertexNumber = HoleEdgeNumber;	// Edge and vertex counts are always the same
-
-		const FVertexID VertexID = GetPolygonHoleVertex( PolygonRef, HoleNumber, HoleVertexNumber );
-		const FVertexID NextVertexID = GetPolygonHoleVertex( PolygonRef, HoleNumber, ( HoleVertexNumber + 1 ) % NumHoleVertices );
-
-		// Find the edge that connects these vertices
-		FEdgeID FoundEdgeID = FEdgeID::Invalid;
-		bool bFoundEdge = false;
-
-		const int32 NumVertexConnectedEdges = GetVertexConnectedEdgeCount( VertexID );	// @todo mesheditor perf: Could be made faster by always iterating over the vertex with the fewest number of connected edges, instead of the first edge vertex
-		for( int32 VertexEdgeNumber = 0; VertexEdgeNumber < NumVertexConnectedEdges; ++VertexEdgeNumber )
-		{
-			const FEdgeID VertexConnectedEdgeID = GetVertexConnectedEdge( VertexID, VertexEdgeNumber );
-
-			// Try the edge's first vertex.  Does it point to us?
-			FVertexID OtherEdgeVertexID = GetEdgeVertex( VertexConnectedEdgeID, 0 );
-			if( OtherEdgeVertexID == VertexID )
-			{
-				// Must be the the other one
-				OtherEdgeVertexID = GetEdgeVertex( VertexConnectedEdgeID, 1 );
-			}
-
-			if( OtherEdgeVertexID == NextVertexID )
-			{
-				// We found the edge!
-				FoundEdgeID = VertexConnectedEdgeID;
-				bFoundEdge = true;
-				break;
-			}
-		}
-		check( bFoundEdge );
-
-		OutHoleEdgeIDs[ HoleEdgeNumber ] = FoundEdgeID;
-	}
+	checkSlow( Polygons.IsAllocated( PolygonID.GetValue() ) );
+	const FMeshPolygon& Polygon = Polygons[ PolygonID.GetValue() ];
+	GetPolygonContourEdges( Polygon.PerimeterContour, OutPolygonPerimeterEdgeIDs );
 }
 
 
-void UEditableMesh::GetPolygonAdjacentPolygons( const FPolygonRef PolygonRef, TArray<FPolygonRef>& OutAdjacentPolygons ) const
+void UEditableMesh::GetPolygonHoleEdges( const FPolygonID PolygonID, const int32 HoleNumber, TArray<FEdgeID>& OutHoleEdgeIDs ) const
 {
+	checkSlow( Polygons.IsAllocated( PolygonID.GetValue() ) );
+	const FMeshPolygon& Polygon = Polygons[ PolygonID.GetValue() ];
+	const FMeshPolygonContour& Contour = Polygon.HoleContours[ HoleNumber ];
+	GetPolygonContourEdges( Contour, OutHoleEdgeIDs );
+}
+
+
+int32 UEditableMesh::GetPolygonTriangulatedTriangleCount( const FPolygonID PolygonID ) const
+{
+	checkSlow( Polygons.IsAllocated( PolygonID.GetValue() ) );
+	const FMeshPolygon& Polygon = Polygons[ PolygonID.GetValue() ];
+
+	return Polygon.Triangles.Num();
+}
+
+
+FMeshTriangle UEditableMesh::GetPolygonTriangulatedTriangle( const FPolygonID PolygonID, int32 PolygonTriangleNumber ) const
+{
+	checkSlow( Polygons.IsAllocated( PolygonID.GetValue() ) );
+	const FMeshPolygon& Polygon = Polygons[ PolygonID.GetValue() ];
+
+	// @todo mesheditor: how to make this safe so that BPs don't crash when given an out-of-range index.
+	// (same goes for just about all of these methods)
+	return Polygon.Triangles[ PolygonTriangleNumber ];
+}
+
+
+FVector UEditableMesh::GetPolygonTriangulatedTriangleVertexPosition( const FPolygonID PolygonID, const int32 PolygonTriangleNumber, const int32 TriangleVertexNumber ) const
+{
+	checkSlow( TriangleVertexNumber >= 0 && TriangleVertexNumber < 3 );
+
+	checkSlow( Polygons.IsAllocated( PolygonID.GetValue() ) );
+	const FMeshPolygon& Polygon = Polygons[ PolygonID.GetValue() ];
+
+	const FVertexInstanceID VertexInstanceID = Polygon.Triangles[ PolygonTriangleNumber ].VertexInstanceIDs[ TriangleVertexNumber ];
+	checkSlow( VertexInstances.IsAllocated( VertexInstanceID.GetValue() ) );
+	const FVertexID VertexID = VertexInstances[ VertexInstanceID.GetValue() ].VertexID;
+
+	checkSlow( Vertices.IsAllocated( VertexID.GetValue() ) );
+	return Vertices[ VertexID.GetValue() ].VertexPosition;
+}
+
+
+void UEditableMesh::GetPolygonAdjacentPolygons( const FPolygonID PolygonID, TArray<FPolygonID>& OutAdjacentPolygons ) const
+{
 	OutAdjacentPolygons.Reset();
 
 	static TArray<FEdgeID> PolygonPerimeterEdges;
-	GetPolygonPerimeterEdges( PolygonRef, PolygonPerimeterEdges );
+	GetPolygonPerimeterEdges( PolygonID, PolygonPerimeterEdges );
 
 	for( const FEdgeID EdgeID : PolygonPerimeterEdges )
 	{
-		static TArray<FPolygonRef> EdgeConnectedPolygons;
+		static TArray<FPolygonID> EdgeConnectedPolygons;
 		GetEdgeConnectedPolygons( EdgeID, EdgeConnectedPolygons );
 
-		for( const FPolygonRef EdgeConnectedPolygon : EdgeConnectedPolygons )
+		for( const FPolygonID EdgeConnectedPolygon : EdgeConnectedPolygons )
 		{
-			if( EdgeConnectedPolygon != PolygonRef )
+			if( EdgeConnectedPolygon != PolygonID )
 			{
 				OutAdjacentPolygons.AddUnique( EdgeConnectedPolygon );
 			}
@@ -672,15 +1833,9 @@
 	FBox BoundingBox;
 	BoundingBox.Init();
 
-	const int32 VertexArraySize = GetVertexArraySize();
-	for( int32 VertexNumber = 0; VertexNumber < VertexArraySize; ++VertexNumber )
+	for( TSparseArray<FMeshVertex>::TConstIterator It( Vertices ); It; ++It )
 	{
-		const FVertexID VertexID( VertexNumber );
-		if( IsValidVertex( VertexID ) )
-		{
-			const FVector VertexPosition = GetVertexAttribute( VertexID, UEditableMeshAttribute::VertexPosition(), 0 );
-			BoundingBox += VertexPosition;
-		}
+		BoundingBox += It->VertexPosition;
 	}
 
 	return BoundingBox;
@@ -713,12 +1868,12 @@
 }
 
 
-FVector UEditableMesh::ComputePolygonCenter( const FPolygonRef PolygonRef ) const
+FVector UEditableMesh::ComputePolygonCenter( const FPolygonID PolygonID ) const
 {
 	FVector Centroid = FVector::ZeroVector;
 
 	static TArray<FVertexID> PerimeterVertexIDs;
-	GetPolygonPerimeterVertices( PolygonRef, /* Out */ PerimeterVertexIDs );
+	GetPolygonPerimeterVertices( PolygonID, /* Out */ PerimeterVertexIDs );
 
 	for( const FVertexID VertexID : PerimeterVertexIDs )
 	{
@@ -731,7 +1886,7 @@
 }
 
 
-FPlane UEditableMesh::ComputePolygonPlane( const FPolygonRef PolygonRef ) const
+FPlane UEditableMesh::ComputePolygonPlane( const FPolygonID PolygonID ) const
 {
 	// NOTE: This polygon plane computation code is partially based on the implementation of "Newell's method" from Real-Time 
 	//       Collision Detection by Christer Ericson, published by Morgan Kaufmann Publishers, (c) 2005 Elsevier Inc
@@ -739,12 +1894,13 @@
 	// @todo mesheditor perf: For polygons that are just triangles, use a cross product to get the normal fast!
 	// @todo mesheditor perf: We could skip computing the plane distance when we only need the normal
 	// @todo mesheditor perf: We could cache these computed polygon normals; or just use the normal of the first three vertices' triangle if it is satisfactory in all cases
+	// @todo mesheditor: For non-planar polygons, the result can vary. Ideally this should use the actual polygon triangulation as opposed to the arbitrary triangulation used here.
 
 	FVector Centroid = FVector::ZeroVector;
 	FVector Normal = FVector::ZeroVector;
 
 	static TArray<FVertexID> PerimeterVertexIDs;
-	GetPolygonPerimeterVertices( PolygonRef, /* Out */ PerimeterVertexIDs );
+	GetPolygonPerimeterVertices( PolygonID, /* Out */ PerimeterVertexIDs );
 
 	// Use 'Newell's Method' to compute a robust 'best fit' plane from the vertices of this polygon
 	for( int32 VertexNumberI = PerimeterVertexIDs.Num() - 1, VertexNumberJ = 0; VertexNumberJ < PerimeterVertexIDs.Num(); VertexNumberI = VertexNumberJ, VertexNumberJ++ )
@@ -769,168 +1925,17 @@
 }
 
 
-FVector UEditableMesh::ComputePolygonNormal( const FPolygonRef PolygonRef ) const
+FVector UEditableMesh::ComputePolygonNormal( const FPolygonID PolygonID ) const
 {
-	const FPlane PolygonPlane = ComputePolygonPlane( PolygonRef );
+	// @todo mesheditor: Polygon normals are now computed and cached when changes are made to a polygon.
+	// In theory, we can just return that cached value, but we need to check that there is nothing which relies on the value being correct before
+	// the cache is updated at the end of a modification.
+	const FPlane PolygonPlane = ComputePolygonPlane( PolygonID );
 	const FVector PolygonNormal( PolygonPlane.X, PolygonPlane.Y, PolygonPlane.Z );
 	return PolygonNormal;
 }
 
 
-FVector UEditableMesh::ComputePolygonPerimeterVertexNormal( const FPolygonRef PolygonRef, const int32 PolygonVertexNumber ) const
-{
-	// The idea here is that we build a graph of adjacent polygons which include the vertex whose normal we want to compute
-	// (where only polygons separated by a soft edge are considered to be adjacent).
-	// Then we walk the graph starting from our start poly, and sum all the visited poly normals, weighted by the angle they make with the vertex.
-	//
-	// We only consider adjacency local to the vertex being considered - this means we can support topology such as a face only partially delimited
-	// by hard edges. Determining smoothing groups would be less flexible.
-
-	// @todo mesheditor: this is not likely to scale well for large meshes (since all normals are recalculated after topology changes)
-	// as it is of pretty high complexity.
-	// We should either cache further adjacency information (e.g. adjacent polys), or not recalculate normals which have not changed.
-
-	// ID of the vertex whose normal we wish to compute
-	const FVertexID VertexID = GetPolygonPerimeterVertex( PolygonRef, PolygonVertexNumber );
-
-	// For every polygon which includes the vertex, we record which polygons are adjacent, and the angle the polygon makes with the vertex.
-	// When walking the graph later, the bVisited flag will be used to ensure we visit each node only once.
-	struct FPolygonGraphNode
-	{
-		FPolygonGraphNode()
-			: Angle( 0.0f )
-			, bVisited( false )
-		{}
-
-		TArray<FPolygonRef> AdjacentPolygonRefs;
-		float Angle;
-		bool bVisited;
-	};
-
-	// Use this to build a graph of polygons
-	static TMap<FPolygonRef, FPolygonGraphNode> PolygonGraphNodes;
-	PolygonGraphNodes.Empty();
-
-	// Get a list of all polygons which share this vertex.
-	// This is the container we iterate through when building the adjacencies.
-	static TArray<FPolygonRef> ConnectedPolygons;
-	GetVertexConnectedPolygons( VertexID, ConnectedPolygons );
-
-	// Get a list of all soft edges which share this vertex.
-	// We're not interested in hard edges as we don't consider them to denote an adjacency between their neighbour polys.
-	static TArray<FEdgeID> ConnectedSoftEdges;
-	ConnectedSoftEdges.Empty();
-
-	const int32 ConnectedEdgeCount = GetVertexConnectedEdgeCount( VertexID );
-	for( int32 EdgeNumber = 0; EdgeNumber < ConnectedEdgeCount; ++EdgeNumber )
-	{
-		const FEdgeID EdgeID( GetVertexConnectedEdge( VertexID, EdgeNumber ) );
-		const bool bIsSoftEdge = FMath::IsNearlyZero( GetEdgeAttribute( EdgeID, UEditableMeshAttribute::EdgeIsHard(), 0 ).X );
-		if( bIsSoftEdge )
-		{
-			ConnectedSoftEdges.Add( EdgeID );
-		}
-	}
-
-	for( FPolygonRef ConnectedPolygon : ConnectedPolygons )
-	{
-		// For each connected polygon, create a PolygonInfo entry
-		FPolygonGraphNode& PolygonGraphNode = PolygonGraphNodes.Add( ConnectedPolygon, FPolygonGraphNode() );
-
-		// Get all the vertices which form the polygon perimeter
-		// @todo mesheditor: what about holes?
-		static TArray<FVertexID> PolygonVertices;
-		GetPolygonPerimeterVertices( ConnectedPolygon, PolygonVertices );
-		const int32 PolygonVertexCount = PolygonVertices.Num();
-		check( PolygonVertexCount > 2 );
-
-		FVertexID LastVertexID( PolygonVertices[ PolygonVertexCount - 2 ] );
-		FVertexID ThisVertexID( PolygonVertices[ PolygonVertexCount - 1 ] );
-
-		// Go through all the polygon perimeter vertices, looking for the vertex whose normal we are setting
-		for( int32 VertexIndex = 0; VertexIndex < PolygonVertexCount; ++VertexIndex )
-		{
-			const FVertexID NextVertexID( PolygonVertices[ VertexIndex ] );
-
-			if( ThisVertexID == VertexID )
-			{
-				// Found the vertex, and we now have the two adjacent vertices too.
-				const FVector LastVertexPosition( GetVertexAttribute( LastVertexID, UEditableMeshAttribute::VertexPosition(), 0 ) );
-				const FVector ThisVertexPosition( GetVertexAttribute( ThisVertexID, UEditableMeshAttribute::VertexPosition(), 0 ) );
-				const FVector NextVertexPosition( GetVertexAttribute( NextVertexID, UEditableMeshAttribute::VertexPosition(), 0 ) );
-
-				const FVector Direction1 = ( LastVertexPosition - ThisVertexPosition ).GetSafeNormal();
-				const FVector Direction2 = ( NextVertexPosition - ThisVertexPosition ).GetSafeNormal();
-
-				// Record the angle the poly makes with the vertex, to be used as the normal weight
-				PolygonGraphNode.Angle = FMath::Acos( FVector::DotProduct( Direction1, Direction2 ) );
-			}
-			else
-			{
-				// Now examine all the soft edges connected to the vertex, looking for those which are also connected to this perimeter vertex.
-				// We need this to determine adjacent polygons.
-				for( FEdgeID ConnectedEdge : ConnectedSoftEdges )
-				{
-					FVertexID EdgeVertexID0;
-					FVertexID EdgeVertexID1;
-					GetEdgeVertices( ConnectedEdge, EdgeVertexID0, EdgeVertexID1 );
-
-					// If we find an edge which contains this vertex ID, look for adjacent polygons.
-					// We know that the other vertex must be the vertex whose normal we're building.
-					if( EdgeVertexID0 == ThisVertexID || EdgeVertexID1 == ThisVertexID )
-					{
-						check( EdgeVertexID0 == VertexID || EdgeVertexID1 == VertexID );
-						const int32 EdgeConnectedPolygonCount = GetEdgeConnectedPolygonCount( ConnectedEdge );
-						for( int32 EdgeConnectedPolygonIndex = 0; EdgeConnectedPolygonIndex < EdgeConnectedPolygonCount; ++EdgeConnectedPolygonIndex )
-						{
-							// One of the adjacent polygons will be the one we are currently processing.
-							// All the others must be adjacent via the edge we are examining.
-							const FPolygonRef EdgeConnectedPolygonRef = GetEdgeConnectedPolygon( ConnectedEdge, EdgeConnectedPolygonIndex );
-							if( EdgeConnectedPolygonRef != ConnectedPolygon )
-							{
-								PolygonGraphNode.AdjacentPolygonRefs.AddUnique( EdgeConnectedPolygonRef );
-							}
-						}
-					}
-				}
-
-			}
-
-			LastVertexID = ThisVertexID;
-			ThisVertexID = NextVertexID;
-		}
-	}
-
-	// Now walk the graph starting from the polygon being set (a simple non-recursive depth-first search)
-	static TArray<FPolygonRef> GraphNodeStack;
-	GraphNodeStack.Empty( PolygonGraphNodes.Num() );
-	GraphNodeStack.Push( PolygonRef );
-
-	FVector Normal( 0.0f );
-
-	while( GraphNodeStack.Num() > 0 )
-	{
-		const FPolygonRef Node = GraphNodeStack.Pop();
-		FPolygonGraphNode* PolygonGraphNode = PolygonGraphNodes.Find( Node );
-
-		if( PolygonGraphNode && !PolygonGraphNode->bVisited )
-		{
-			// Make node as visited and add normal contribution to result
-			PolygonGraphNode->bVisited = true;
-			Normal += ComputePolygonNormal( Node ) * PolygonGraphNode->Angle;
-
-			// Add adjacent polygons to the node stack to be walked
-			for( FPolygonRef AdjacentPolygonRef : PolygonGraphNode->AdjacentPolygonRefs )
-			{
-				GraphNodeStack.Push( AdjacentPolygonRef );
-			}
-		}
-	}
-
-	return Normal.GetSafeNormal();
-}
-
-
 void UEditableMesh::RefreshOpenSubdiv()
 {
 	OsdTopologyRefiner.Reset();
@@ -944,7 +1949,7 @@
 		{
 			const int32 VertexArraySize = GetVertexArraySize();
 			OsdTopologyDescriptor.numVertices = VertexArraySize;
-			OsdTopologyDescriptor.numFaces = GetTotalPolygonCount();
+			OsdTopologyDescriptor.numFaces = GetPolygonCount();
 
 			// NOTE: OpenSubdiv likes weights to be between 0.0 and 10.0, so we'll account for that here
 			const float OpenSubdivCreaseWeightMultiplier = 10.0f;
@@ -1002,32 +2007,24 @@
 
 				int32 NextOsdFaceIndex = 0;
 
-				const int32 SectionArraySize = GetSectionArraySize();
-				for( int32 SectionNumber = 0; SectionNumber < SectionArraySize; ++SectionNumber )
+				const int32 PolygonCount = GetPolygonArraySize();
+				for( int32 PolygonIndex = 0; PolygonIndex < PolygonCount; ++PolygonIndex )
 				{
-					const FSectionID SectionID( SectionNumber );
-					if( IsValidSection( SectionID ) )
+					const FPolygonID PolygonID( PolygonIndex );
+					if( IsValidPolygon( PolygonID ) )
 					{
-						int32 SectionPolygonArraySize = GetPolygonArraySize( SectionID );
-						for( int32 SectionPolygonNumber = 0; SectionPolygonNumber < SectionPolygonArraySize; ++SectionPolygonNumber )
+						static TArray<FVertexID> PerimeterVertexIDs;
+						GetPolygonPerimeterVertices( PolygonID, /* Out */ PerimeterVertexIDs );
+
+						const int32 PerimeterVertexCount = PerimeterVertexIDs.Num();
+						OsdNumVerticesPerFace[ NextOsdFaceIndex++ ] = PerimeterVertexCount;
+
+						for( int32 PerimeterVertexNumber = 0; PerimeterVertexNumber < PerimeterVertexCount; ++PerimeterVertexNumber )
 						{
-							const FPolygonRef PolygonRef( SectionID, FPolygonID( SectionPolygonNumber ) );
-							if( IsValidPolygon( PolygonRef ) )
-							{
-								static TArray<FVertexID> PerimeterVertexIDs;
-								GetPolygonPerimeterVertices( PolygonRef, /* Out */ PerimeterVertexIDs );
+							const FVertexID PerimeterVertexID = PerimeterVertexIDs[ PerimeterVertexNumber ];
 
-								const int32 PerimeterVertexCount = PerimeterVertexIDs.Num();
-								OsdNumVerticesPerFace[ NextOsdFaceIndex++ ] = PerimeterVertexCount;
-
-								for( int32 PerimeterVertexNumber = 0; PerimeterVertexNumber < PerimeterVertexCount; ++PerimeterVertexNumber )
-								{
-									const FVertexID PerimeterVertexID = PerimeterVertexIDs[ PerimeterVertexNumber ];
-
-									OsdVertexIndicesPerFace.Add( PerimeterVertexID.GetValue() );
-									OsdFVarIndicesPerFace.Add( OsdFVarIndicesPerFace.Num() );
-								}
-							}
+							OsdVertexIndicesPerFace.Add( PerimeterVertexID.GetValue() );
+							OsdFVarIndicesPerFace.Add( OsdFVarIndicesPerFace.Num() );
 						}
 					}
 				}
@@ -1158,7 +2155,7 @@
 			FLinearColor VertexColor;
 		};
 
-		const int32 SectionCount = GetSectionCount();
+		const int32 PolygonGroupCount = GetPolygonGroupCount();
 
 		// Get the limit surface subdivision level from OpenSubdiv
 		const OpenSubdiv::Far::TopologyLevel& OsdLimitLevel = OsdTopologyRefiner->GetLevel( SubdivisionCount );
@@ -1174,7 +2171,7 @@
 		// Grab all of the vertex data and put them in separate contiguous arrays for OpenSubdiv
 		static TArray<FVector> VertexPositions;
 		static TArray<FOsdFVarVertexData> FVarVertexDatas;
-		static TArray<int32> FirstPolygonNumberForSections;
+		static TArray<int32> FirstPolygonNumberForPolygonGroups;
 		{
 			// Vertex positions
 			{
@@ -1198,37 +2195,35 @@
 			}
 
 			// Texture coordinates (per polygon vertex)
-			FirstPolygonNumberForSections.Reset();
+			FirstPolygonNumberForPolygonGroups.Reset();
 
 			{
 				FVarVertexDatas.Reset();
 				FVarVertexDatas.Reserve( OsdFVarIndicesPerFace.Num() );
 
-				const int32 SectionArraySize = GetSectionArraySize();
+				const int32 PolygonGroupArraySize = GetPolygonGroupArraySize();
 
 				int32 NumPolygonsSoFar = 0;
-				for( int32 SectionNumber = 0; SectionNumber < SectionArraySize; ++SectionNumber )
+				for( int32 PolygonGroupNumber = 0; PolygonGroupNumber < PolygonGroupArraySize; ++PolygonGroupNumber )
 				{
-					const FSectionID SectionID( SectionNumber );
-					if( IsValidSection( SectionID ) )
+					const FPolygonGroupID PolygonGroupID( PolygonGroupNumber );
+					if( IsValidPolygonGroup( PolygonGroupID ) )
 					{
-						FirstPolygonNumberForSections.Add( NumPolygonsSoFar );
-						NumPolygonsSoFar += GetPolygonCount( SectionID );
+						FirstPolygonNumberForPolygonGroups.Add( NumPolygonsSoFar );
+						const int32 PolygonCountInGroup = GetPolygonCountInGroup( PolygonGroupID );
+						NumPolygonsSoFar += PolygonCountInGroup;
 
-						int32 SectionPolygonArraySize = GetPolygonArraySize( SectionID );
-						for( int32 SectionPolygonNumber = 0; SectionPolygonNumber < SectionPolygonArraySize; ++SectionPolygonNumber )
+						for( int32 PolygonGroupPolygonNumber = 0; PolygonGroupPolygonNumber < PolygonCountInGroup; ++PolygonGroupPolygonNumber )
 						{
-							const FPolygonRef PolygonRef( SectionID, FPolygonID( SectionPolygonNumber ) );
-							if( IsValidPolygon( PolygonRef ) )
+							const FPolygonID PolygonID( GetPolygonInGroup( PolygonGroupID, PolygonGroupPolygonNumber ) );
+
+							const int32 PerimeterVertexCount = GetPolygonPerimeterVertexCount( PolygonID );
+							for( int32 PerimeterVertexNumber = 0; PerimeterVertexNumber < PerimeterVertexCount; ++PerimeterVertexNumber )
 							{
-								const int32 PerimeterVertexCount = GetPolygonPerimeterVertexCount( PolygonRef );
-								for( int32 PerimeterVertexNumber = 0; PerimeterVertexNumber < PerimeterVertexCount; ++PerimeterVertexNumber )
-								{
-									FOsdFVarVertexData& FVarVertexData = *new( FVarVertexDatas ) FOsdFVarVertexData();
-									FVarVertexData.TextureCoordinates[ 0 ] = TextureCoordinateCount > 0 ? FVector2D( GetPolygonPerimeterVertexAttribute( PolygonRef, PerimeterVertexNumber, UEditableMeshAttribute::VertexTextureCoordinate(), 0 ) ) : FVector2D::ZeroVector;
-									FVarVertexData.TextureCoordinates[ 1 ] = TextureCoordinateCount > 1 ? FVector2D( GetPolygonPerimeterVertexAttribute( PolygonRef, PerimeterVertexNumber, UEditableMeshAttribute::VertexTextureCoordinate(), 1 ) ) : FVector2D::ZeroVector;
-									FVarVertexData.VertexColor = FLinearColor( GetPolygonPerimeterVertexAttribute( PolygonRef, PerimeterVertexNumber, UEditableMeshAttribute::VertexColor(), 0 ) );
-								}
+								FOsdFVarVertexData& FVarVertexData = *new( FVarVertexDatas ) FOsdFVarVertexData();
+								FVarVertexData.TextureCoordinates[ 0 ] = TextureCoordinateCount > 0 ? FVector2D( GetPolygonPerimeterVertexAttribute( PolygonID, PerimeterVertexNumber, UEditableMeshAttribute::VertexTextureCoordinate(), 0 ) ) : FVector2D::ZeroVector;
+								FVarVertexData.TextureCoordinates[ 1 ] = TextureCoordinateCount > 1 ? FVector2D( GetPolygonPerimeterVertexAttribute( PolygonID, PerimeterVertexNumber, UEditableMeshAttribute::VertexTextureCoordinate(), 1 ) ) : FVector2D::ZeroVector;
+								FVarVertexData.VertexColor = FLinearColor( GetPolygonPerimeterVertexAttribute( PolygonID, PerimeterVertexNumber, UEditableMeshAttribute::VertexColor(), 0 ) );
 							}
 						}
 					}
@@ -1350,7 +2345,7 @@
 			}
 		}
 
-		SubdivisionLimitData.Sections.SetNum( SectionCount, false );
+		SubdivisionLimitData.Sections.SetNum( PolygonGroupCount, false );
 
 		for( int32 LimitFaceNumber = 0; LimitFaceNumber < LimitFaceCount; ++LimitFaceNumber )
 		{
@@ -1372,11 +2367,11 @@
 				}
 				int32 BaseCageFaceNumber = CurrentFaceNumber;
 
-				for( int32 SectionNumber = SectionCount - 1; SectionNumber >= 0; --SectionNumber )
+				for( int32 PolygonGroupNumber = PolygonGroupCount - 1; PolygonGroupNumber >= 0; --PolygonGroupNumber)
 				{
-					if( BaseCageFaceNumber >= FirstPolygonNumberForSections[ SectionNumber ] )
+					if( BaseCageFaceNumber >= FirstPolygonNumberForPolygonGroups[ PolygonGroupNumber ] )
 					{
-						QuadSectionNumber = SectionNumber;
+						QuadSectionNumber = PolygonGroupNumber;
 						break;
 					}
 				}
@@ -1591,8 +2586,24 @@
 }
 
 
-void UEditableMesh::ComputePolygonTriangulation( const FPolygonRef PolygonRef, TArray<int32>& OutPerimeterVertexNumbersForTriangles ) const
+void UEditableMesh::RetriangulatePolygons()
 {
+	// Perform triangulation directly into mesh polygons
+	for( const FPolygonID PolygonID : PolygonsPendingTriangulation )
+	{
+		FMeshPolygon& Polygon = Polygons[ PolygonID.GetValue() ];
+		ComputePolygonTriangulation( PolygonID, Polygon.Triangles );
+	}
+
+	for( UEditableMeshAdapter* Adapter : Adapters )
+	{
+		Adapter->OnRetriangulatePolygons( this, PolygonsPendingTriangulation.Array() );
+	}
+}
+
+
+void UEditableMesh::ComputePolygonTriangulation( const FPolygonID PolygonID, TArray<FMeshTriangle>& OutTriangles ) const
+{
 	// NOTE: This polygon triangulation code is partially based on the ear cutting algorithm described on
 	//       page 497 of the book "Real-time Collision Detection", published in 2005.
 
@@ -1610,19 +2621,21 @@
 	};
 
 
-	OutPerimeterVertexNumbersForTriangles.Reset();
+	OutTriangles.Reset();
 
 	// @todo mesheditor holes: Does not support triangles with holes yet!
+	// @todo mesheditor: Perhaps should always attempt to triangulate by splitting polygons along the shortest edge, for better determinism.
 
+	const FMeshPolygon& Polygon = Polygons[ PolygonID.GetValue() ];
+	const TArray<FVertexInstanceID>& PolygonVertexInstanceIDs = Polygon.PerimeterContour.VertexInstanceIDs;
+
 	// Polygon must have at least three vertices/edges
-	static TArray<FVertexID> VertexIDs;
-	GetPolygonPerimeterVertices( PolygonRef, /* Out */ VertexIDs );
-	const int32 PolygonVertexCount = VertexIDs.Num();
+	const int32 PolygonVertexCount = PolygonVertexInstanceIDs.Num();
 	check( PolygonVertexCount >= 3 );
 
 	// First figure out the polygon normal.  We need this to determine which triangles are convex, so that
 	// we can figure out which ears to clip
-	const FVector PolygonNormal = ComputePolygonNormal( PolygonRef );
+	const FVector PolygonNormal = ComputePolygonNormal( PolygonID );
 
 	// Make a simple linked list array of the previous and next vertex numbers, for each vertex number
 	// in the polygon.  This will just save us having to iterate later on.
@@ -1638,7 +2651,7 @@
 			PrevVertexNumbers[ VertexNumber ] = VertexNumber - 1;
 			NextVertexNumbers[ VertexNumber ] = VertexNumber + 1;
 
-			VertexPositions[ VertexNumber ] = GetVertexAttribute( VertexIDs[ VertexNumber ], UEditableMeshAttribute::VertexPosition(), 0 );
+			VertexPositions[ VertexNumber ] = GetVertexInstanceAttribute( PolygonVertexInstanceIDs[ VertexNumber ], UEditableMeshAttribute::VertexPosition(), 0 );
 		}
 		PrevVertexNumbers[ 0 ] = PolygonVertexCount - 1;
 		NextVertexNumbers[ PolygonVertexCount - 1 ] = 0;
@@ -1700,11 +2713,12 @@
 		{
 			// OK, we found an ear!  Let's save this triangle in our output buffer.
 			{
-				const int32 TriangleVertexIDNumbers[ 3 ] = { PrevVertexNumbers[ EarVertexNumber ], EarVertexNumber, NextVertexNumbers[ EarVertexNumber ] };
-				for( int32 TriangleVertexNumber = 0; TriangleVertexNumber < 3; ++TriangleVertexNumber )
-				{
-					OutPerimeterVertexNumbersForTriangles.Add( TriangleVertexIDNumbers[ TriangleVertexNumber ] );
-				}
+				OutTriangles.Emplace();
+				FMeshTriangle& Triangle = OutTriangles.Last();
+
+				Triangle.VertexInstanceIDs[ 0 ] = PolygonVertexInstanceIDs[ PrevVertexNumbers[ EarVertexNumber ] ];
+				Triangle.VertexInstanceIDs[ 1 ] = PolygonVertexInstanceIDs[ EarVertexNumber ];
+				Triangle.VertexInstanceIDs[ 2 ] = PolygonVertexInstanceIDs[ NextVertexNumbers[ EarVertexNumber ] ];
 			}
 
 			// Update our linked list.  We're effectively cutting off the ear by pointing the ear vertex's neighbors to
@@ -1732,56 +2746,31 @@
 		}
 	}
 
-	check( OutPerimeterVertexNumbersForTriangles.Num() > 0 );
-	check( OutPerimeterVertexNumbersForTriangles.Num() % 3 == 0 );	// Must all be triangles
+	check( OutTriangles.Num() > 0 );
 }
 
 
-void UEditableMesh::ComputePolygonTriangulation( const FPolygonRef PolygonRef, TArray<FVertexID>& OutTrianglesVertexIDs ) const
+bool UEditableMesh::ComputeBarycentricWeightForPointOnPolygon( const FPolygonID PolygonID, const FVector PointOnPolygon, FMeshTriangle& OutTriangle, FVector& OutTriangleVertexWeights ) const
 {
-	static TArray<int32> PerimeterVertexNumbersForTriangles;
-	ComputePolygonTriangulation( PolygonRef, /* Out */ PerimeterVertexNumbersForTriangles );
+	// @todo mesheditor: Modify this method so it can return a meaningful barycentric weight for an off-polygon position, referencing the closest triangle, with out-of-range weights.
 
-	static TArray<FVertexID> PerimeterVertexIDs;
-	GetPolygonPerimeterVertices( PolygonRef, /* Out */ PerimeterVertexIDs );
+	const FMeshPolygon& Polygon = Polygons[ PolygonID.GetValue() ];
+	const int32 TriangleCount = Polygon.Triangles.Num();
 
-	OutTrianglesVertexIDs.SetNum( PerimeterVertexNumbersForTriangles.Num(), false );
-	for( int32 TriangleVerticesNumber = 0; TriangleVerticesNumber < PerimeterVertexNumbersForTriangles.Num(); ++TriangleVerticesNumber )
-	{
-		const int32 PerimeterVertexNumber = PerimeterVertexNumbersForTriangles[ TriangleVerticesNumber ];
-		OutTrianglesVertexIDs[ TriangleVerticesNumber ] = PerimeterVertexIDs[ PerimeterVertexNumber ];
-	}
-}
-
-bool UEditableMesh::ComputeBarycentricWeightForPointOnPolygon( const FPolygonRef PolygonRef, const FVector PointOnPolygon, TArray<int32>& OutPerimeterVertexIndices, FVector& OutTriangleVertexWeights ) const
-{
-	OutPerimeterVertexIndices.Reset( 3 );
-
-	// Triangulate the polygon, so that we can compute barycentric weights for the texture coordinates that will be used
-	// for the new inset vertices
-	static TArray<int32> PerimeterVertexNumbersForTriangles;
-	ComputePolygonTriangulation( PolygonRef, /* Out */ PerimeterVertexNumbersForTriangles );
-
-	check( PerimeterVertexNumbersForTriangles.Num() % 3 == 0 );
-	const int32 TriangleCount = PerimeterVertexNumbersForTriangles.Num() / 3;
-
 	// Figure out which triangle the incoming point is within
-	for( int32 TriangleNumber = 0; TriangleNumber < TriangleCount; ++TriangleNumber )
+	for( const FMeshTriangle& Triangle : Polygon.Triangles )
 	{
-		const FVector TriangleVertex0Position = GetPolygonPerimeterVertexAttribute( PolygonRef, PerimeterVertexNumbersForTriangles[ TriangleNumber * 3 + 0 ], UEditableMeshAttribute::VertexPosition(), 0 );
-		const FVector TriangleVertex1Position = GetPolygonPerimeterVertexAttribute( PolygonRef, PerimeterVertexNumbersForTriangles[ TriangleNumber * 3 + 1 ], UEditableMeshAttribute::VertexPosition(), 0 );
-		const FVector TriangleVertex2Position = GetPolygonPerimeterVertexAttribute( PolygonRef, PerimeterVertexNumbersForTriangles[ TriangleNumber * 3 + 2 ], UEditableMeshAttribute::VertexPosition(), 0 );
+		const FVector TriangleVertex0Position = GetVertexInstanceAttribute( Triangle.VertexInstanceIDs[ 0 ], UEditableMeshAttribute::VertexPosition(), 0 );
+		const FVector TriangleVertex1Position = GetVertexInstanceAttribute( Triangle.VertexInstanceIDs[ 1 ], UEditableMeshAttribute::VertexPosition(), 0 );
+		const FVector TriangleVertex2Position = GetVertexInstanceAttribute( Triangle.VertexInstanceIDs[ 2 ], UEditableMeshAttribute::VertexPosition(), 0 );
 
 		// Calculate the barycentric weights for the triangle's verts and determine if the point lies within its bounds.
 		OutTriangleVertexWeights = FMath::ComputeBaryCentric2D( PointOnPolygon, TriangleVertex0Position, TriangleVertex1Position, TriangleVertex2Position );
 
 		if( OutTriangleVertexWeights.X >= 0.0f && OutTriangleVertexWeights.Y >= 0.0f && OutTriangleVertexWeights.Z >= 0.0f )
 		{
-			// Okay, we found a triangle that the point is inside!  Return the perimeter vertex indices.
-			OutPerimeterVertexIndices.Add( PerimeterVertexNumbersForTriangles[ TriangleNumber * 3 + 0 ] );
-			OutPerimeterVertexIndices.Add( PerimeterVertexNumbersForTriangles[ TriangleNumber * 3 + 1 ] );
-			OutPerimeterVertexIndices.Add( PerimeterVertexNumbersForTriangles[ TriangleNumber * 3 + 2 ] );
-
+			// Okay, we found a triangle that the point is inside!  Return the corresponding vertex instances.
+			OutTriangle = Triangle;
 			return true;
 		}
 	}
@@ -1790,23 +2779,23 @@
 }
 
 
-void UEditableMesh::ComputeTextureCoordinatesForPointOnPolygon( const FPolygonRef PolygonRef, const FVector PointOnPolygon, bool& bOutWereTextureCoordinatesFound, TArray<FVector4>& OutInterpolatedTextureCoordinates )
+void UEditableMesh::ComputeTextureCoordinatesForPointOnPolygon( const FPolygonID PolygonID, const FVector PointOnPolygon, bool& bOutWereTextureCoordinatesFound, TArray<FVector4>& OutInterpolatedTextureCoordinates )
 {
 	bOutWereTextureCoordinatesFound = false;
 	OutInterpolatedTextureCoordinates.Reset();
 
-	static TArray<int32> PerimeterVertexNumbers;
+	FMeshTriangle Triangle;
 	FVector TriangleVertexWeights;
-	if( ComputeBarycentricWeightForPointOnPolygon( PolygonRef, PointOnPolygon, PerimeterVertexNumbers, TriangleVertexWeights ) )
+	if( ComputeBarycentricWeightForPointOnPolygon( PolygonID, PointOnPolygon, Triangle, TriangleVertexWeights ) )
 	{
 		OutInterpolatedTextureCoordinates.SetNum( TextureCoordinateCount, false );
 
 		for( int32 TextureCoordinateIndex = 0; TextureCoordinateIndex < TextureCoordinateCount; ++TextureCoordinateIndex )
 		{
 			OutInterpolatedTextureCoordinates[ TextureCoordinateIndex ] =
-				TriangleVertexWeights.X * GetPolygonPerimeterVertexAttribute( PolygonRef, PerimeterVertexNumbers[ 0 ], UEditableMeshAttribute::VertexTextureCoordinate(), TextureCoordinateIndex ) +
-				TriangleVertexWeights.Y * GetPolygonPerimeterVertexAttribute( PolygonRef, PerimeterVertexNumbers[ 1 ], UEditableMeshAttribute::VertexTextureCoordinate(), TextureCoordinateIndex ) +
-				TriangleVertexWeights.Z * GetPolygonPerimeterVertexAttribute( PolygonRef, PerimeterVertexNumbers[ 2 ], UEditableMeshAttribute::VertexTextureCoordinate(), TextureCoordinateIndex );
+				TriangleVertexWeights.X * GetVertexInstanceAttribute( Triangle.VertexInstanceIDs[ 0 ], UEditableMeshAttribute::VertexTextureCoordinate(), TextureCoordinateIndex ) +
+				TriangleVertexWeights.Y * GetVertexInstanceAttribute( Triangle.VertexInstanceIDs[ 1 ], UEditableMeshAttribute::VertexTextureCoordinate(), TextureCoordinateIndex ) +
+				TriangleVertexWeights.Z * GetVertexInstanceAttribute( Triangle.VertexInstanceIDs[ 2 ], UEditableMeshAttribute::VertexTextureCoordinate(), TextureCoordinateIndex );
 		}
 
 		bOutWereTextureCoordinatesFound = true;
@@ -1844,8 +2833,8 @@
 
 void UEditableMesh::MoveVertices( const TArray<FVertexToMove>& VerticesToMove )
 {
-	static TSet< FPolygonRef > PolygonsThatNeedNewNormals;
-	PolygonsThatNeedNewNormals.Reset();
+	static TSet< FPolygonID > VertexConnectedPolygons;
+	VertexConnectedPolygons.Reset();
 
 	static TArray<FAttributesForVertex> VertexAttributesToSet;
 	VertexAttributesToSet.Reset();
@@ -1861,53 +2850,36 @@
 			FMeshElementAttributeData& VertexAttribute = *new( AttributesForVertex.VertexAttributes.Attributes ) FMeshElementAttributeData( UEditableMeshAttribute::VertexPosition(), 0, VertexToMove.NewVertexPosition );
 
 			// All of the polygons that share this vertex will need new normals
-			static TArray<FPolygonRef> ConnectedPolygonRefs;
+			static TArray<FPolygonID> ConnectedPolygonRefs;
 			GetVertexConnectedPolygons( VertexToMove.VertexID, /* Out */ ConnectedPolygonRefs );
-			PolygonsThatNeedNewNormals.Append( ConnectedPolygonRefs );
+			VertexConnectedPolygons.Append( ConnectedPolygonRefs );
 		}
 	}
 
-	// Convert from a TSet to a TArray
-	static TArray<FPolygonRef> PolygonsToUpdate;
-	PolygonsToUpdate.Reset();
-	for( const FPolygonRef PolygonRef : PolygonsThatNeedNewNormals )
-	{
-		PolygonsToUpdate.Add( PolygonRef );
-	}
-
-	// Perform retriangulation after the geometry is moved back, when undoing only
-	{
-		const bool bOnlyOnUndo = true;
-		RetriangulatePolygons( PolygonsToUpdate, bOnlyOnUndo );
-	}
-
 	SetVerticesAttributes( VertexAttributesToSet );
 
-	// Generate new normals and tangents for any geometry that was moved around
-	GenerateNormalsAndTangentsForPolygonsAndAdjacents( PolygonsToUpdate );
-
+	// Mark all polygons connected to the vertex as requiring a new tangent basis and retriangulation
 	// Everything needs to be retriangulated because convexity may have changed
-	// @todo mesheditor perf: Polygons will actually be triangulated twice in the same frame when performing an operation that implicitly retriangulates at the end of the op (once before move, once after move.)  Not super ideal.
-	{
-		const bool bOnlyOnUndo = false;
-		RetriangulatePolygons( PolygonsToUpdate, bOnlyOnUndo );
-	}
+	PolygonsPendingNewTangentBasis.Append( VertexConnectedPolygons );
+	PolygonsPendingTriangulation.Append( VertexConnectedPolygons );
 }
 
 
-void UEditableMesh::CreateMissingPolygonPerimeterEdges( const FPolygonRef PolygonRef, TArray<FEdgeID>& OutNewEdgeIDs )
+void UEditableMesh::CreateMissingPolygonPerimeterEdges( const FPolygonID PolygonID, TArray<FEdgeID>& OutNewEdgeIDs )
 {
+	// @todo mesheditor: Orphaned method - should we keep it or is it entirely unnecessary now?
+
 	OutNewEdgeIDs.Reset();
 
-	const int32 NumPolygonPerimeterEdges = GetPolygonPerimeterEdgeCount( PolygonRef );
+	const int32 NumPolygonPerimeterEdges = GetPolygonPerimeterEdgeCount( PolygonID );
 	const int32 NumPolygonPerimeterVertices = NumPolygonPerimeterEdges;		// Edge and vertex count are always the same
 
 	for( int32 PerimeterEdgeNumber = 0; PerimeterEdgeNumber < NumPolygonPerimeterEdges; ++PerimeterEdgeNumber )
 	{
 		const int32 PerimeterVertexNumber = PerimeterEdgeNumber;	// Edge and vertex counts are always the same
 
-		const FVertexID VertexID = GetPolygonPerimeterVertex( PolygonRef, PerimeterVertexNumber );
-		const FVertexID NextVertexID = GetPolygonPerimeterVertex( PolygonRef, ( PerimeterVertexNumber + 1 ) % NumPolygonPerimeterVertices );
+		const FVertexID VertexID = GetPolygonPerimeterVertex( PolygonID, PerimeterVertexNumber );
+		const FVertexID NextVertexID = GetPolygonPerimeterVertex( PolygonID, ( PerimeterVertexNumber + 1 ) % NumPolygonPerimeterVertices );
 
 		// Find the edge that connects these vertices
 		FEdgeID FoundEdgeID = FEdgeID::Invalid;
@@ -1956,3 +2928,727 @@
 
 			OutNewEdgeIDs.Append( NewEdgeIDs );
 		}
+	}
+}
+
+
+void UEditableMesh::CreateMissingPolygonHoleEdges( const FPolygonID PolygonID, const int32 HoleNumber, TArray<FEdgeID>& OutNewEdgeIDs )
+{
+	// @todo mesheditor: Orphaned method - should we keep it or is it entirely unnecessary now?
+
+	OutNewEdgeIDs.Reset();
+
+	const int32 NumHoleEdges = GetPolygonHoleEdgeCount( PolygonID, HoleNumber );
+	const int32 NumHoleVertices = NumHoleEdges;		// Edge and vertex count are always the same
+
+	for( int32 HoleEdgeNumber = 0; HoleEdgeNumber < NumHoleEdges; ++HoleEdgeNumber )
+	{
+		const int32 HoleVertexNumber = HoleEdgeNumber;	// Edge and vertex counts are always the same
+
+		const FVertexID VertexID = GetPolygonHoleVertex( PolygonID, HoleNumber, HoleVertexNumber );
+		const FVertexID NextVertexID = GetPolygonHoleVertex( PolygonID, HoleNumber, ( HoleVertexNumber + 1 ) % NumHoleVertices );
+
+		// Find the edge that connects these vertices
+		FEdgeID FoundEdgeID = FEdgeID::Invalid;
+		{
+			bool bFoundEdge = false;
+
+			const int32 NumVertexConnectedEdges = GetVertexConnectedEdgeCount( VertexID );	// @todo mesheditor perf: Could be made faster by always iterating over the vertex with the fewest number of connected edges, instead of the first edge vertex
+			for( int32 VertexEdgeNumber = 0; VertexEdgeNumber < NumVertexConnectedEdges; ++VertexEdgeNumber )
+			{
+				const FEdgeID VertexConnectedEdgeID = GetVertexConnectedEdge( VertexID, VertexEdgeNumber );
+
+				// Try the edge's first vertex.  Does it point to us?
+				FVertexID OtherEdgeVertexID = GetEdgeVertex( VertexConnectedEdgeID, 0 );
+				if( OtherEdgeVertexID == VertexID )
+				{
+					// Must be the the other one
+					OtherEdgeVertexID = GetEdgeVertex( VertexConnectedEdgeID, 1 );
+				}
+
+				if( OtherEdgeVertexID == NextVertexID )
+				{
+					// We found the edge!
+					FoundEdgeID = VertexConnectedEdgeID;
+					bFoundEdge = true;
+					break;
+				}
+			}
+
+			if( !bFoundEdge )
+			{
+				// Create the new edge!  Note that this does not connect the edge to the polygon.  We expect the caller to do that afterwards.
+				FEdgeToCreate EdgeToCreate;
+				EdgeToCreate.VertexID0 = VertexID;
+				EdgeToCreate.VertexID1 = NextVertexID;
+
+				static TArray< FEdgeToCreate > EdgesToCreate;
+				EdgesToCreate.Reset();
+				EdgesToCreate.Add( EdgeToCreate );
+
+				static TArray< FEdgeID > NewEdgeIDs;
+				NewEdgeIDs.Reset();
+				CreateEdges( EdgesToCreate, /* Out */ NewEdgeIDs );
+
+				OutNewEdgeIDs.Append( NewEdgeIDs );
+			}
+		}
+	}
+}
+
+
+void UEditableMesh::SplitEdge( const FEdgeID EdgeID, const TArray<float>& Splits, TArray<FVertexID>& OutNewVertexIDs )
+{
+	UE_LOG( LogEditableMesh, Verbose, TEXT( "SplitEdge: %s" ), *EdgeID.ToString() );
+
+	// NOTE: The incoming splits should always be between 0.0 and 1.0, representing progress along 
+	//       the edge from the edge's first vertex toward it's other vertex.  The order doesn't matter (we'll sort them.)
+
+	check( Splits.Num() > 0 );
+
+	// Sort the split values smallest to largest.  We'll be adding a strip of vertices for each split, and the
+	// indices for those new vertices need to be in order.
+	static TArray<float> SortedSplits;
+	SortedSplits.Reset();
+	SortedSplits = Splits;
+	if( SortedSplits.Num() > 1 )
+	{
+		SortedSplits.Sort();
+	}
+
+	FVertexID OriginalEdgeVertexIDs[ 2 ];
+	GetEdgeVertices( EdgeID, /* Out */ OriginalEdgeVertexIDs[0], /* Out */ OriginalEdgeVertexIDs[1] );
+
+	static TArray<FVector> NewVertexPositions;
+	NewVertexPositions.Reset();
+	NewVertexPositions.AddUninitialized( SortedSplits.Num() );
+
+	for( int32 NewVertexNumber = 0; NewVertexNumber < NewVertexPositions.Num(); ++NewVertexNumber )
+	{
+		const float Split = SortedSplits[ NewVertexNumber ];
+		check( Split >= 0.0f && Split <= 1.0f );
+
+		NewVertexPositions[ NewVertexNumber ] = FMath::Lerp(
+			FVector( GetVertexAttribute( OriginalEdgeVertexIDs[ 0 ], UEditableMeshAttribute::VertexPosition(), 0 ) ),
+			FVector( GetVertexAttribute( OriginalEdgeVertexIDs[ 1 ], UEditableMeshAttribute::VertexPosition(), 0 ) ),
+			Split );
+	}
+
+	// Split the edge, and connect the vertex to the polygons that share the two new edges
+	const FVertexID OriginalEdgeFarVertexID = OriginalEdgeVertexIDs[ 1 ];
+
+	static TArray<FMeshElementAttributeData> OriginalEdgeAttributes;
+	OriginalEdgeAttributes.Reset();
+	{
+		for( const FName AttributeName : UEditableMesh::GetValidEdgeAttributes() )
+		{
+			const int32 MaxAttributeIndex = GetMaxAttributeIndex( AttributeName );
+			for( int32 AttributeIndex = 0; AttributeIndex < MaxAttributeIndex; ++AttributeIndex )
+			{
+				FMeshElementAttributeData& EdgeAttribute = *new( OriginalEdgeAttributes ) FMeshElementAttributeData(
+					AttributeName,
+					AttributeIndex,
+					GetEdgeAttribute( EdgeID, AttributeName, AttributeIndex ) );
+			}
+		}
+	}
+
+
+	// Create new vertices.  Their texture coordinates will be undefined, for now.  We'll set that later by using the connected 
+	// polygon's vertex data.  Also, the vertex's normals and tangents will be automatically generated.
+	static TArray<FVertexID> NewVertexIDs;
+	NewVertexIDs.Reset();
+	{
+		OutNewVertexIDs.Reset();
+		OutNewVertexIDs.Reserve( NewVertexPositions.Num() );
+
+		static TArray<FVertexToCreate> VerticesToCreate;
+		VerticesToCreate.Reset();
+		VerticesToCreate.Reserve( NewVertexPositions.Num() );
+
+		for( int32 NewVertexNumber = 0; NewVertexNumber < NewVertexPositions.Num(); ++NewVertexNumber )
+		{
+			const FVector& NewVertexPosition = NewVertexPositions[ NewVertexNumber ];
+
+			FVertexToCreate& VertexToCreate = *new( VerticesToCreate )FVertexToCreate();
+			VertexToCreate.VertexAttributes.Attributes.Add( FMeshElementAttributeData( UEditableMeshAttribute::VertexPosition(), 0, NewVertexPosition ) );
+		}
+
+		CreateVertices( VerticesToCreate, /* Out */ NewVertexIDs );
+
+		OutNewVertexIDs.Append( NewVertexIDs );
+	}
+
+	// Figure out which polygons are connected to the original edge.  Those polygons are going to need their
+	// connected vertices updated later on.
+	struct FAffectedPolygonEdge
+	{
+		FPolygonID PolygonID;
+		int32 PolygonVertexNumbers[ 2 ];
+	};
+	static TArray<FAffectedPolygonEdge> AffectedPolygonEdges;
+	AffectedPolygonEdges.Reset();
+
+	static TArray<FPolygonID> AffectedPolygons;
+	AffectedPolygons.Reset();
+
+	{
+		const int32 ConnectedPolygonCount = GetEdgeConnectedPolygonCount( EdgeID );
+		for( int32 ConnectedPolygonNumber = 0; ConnectedPolygonNumber < ConnectedPolygonCount; ++ConnectedPolygonNumber )
+		{
+			const FPolygonID PolygonID = GetEdgeConnectedPolygon( EdgeID, ConnectedPolygonNumber );
+
+			// Figure out which of this polygon's vertices are part of the edge we're adding a vertex to
+			// @todo mesheditor hole: Add support for polygon holes, too!
+			static TArray<FVertexID> PolygonPerimeterVertexIDs;
+			PolygonPerimeterVertexIDs.Reset();
+			GetPolygonPerimeterVertices( PolygonID, /* Out */ PolygonPerimeterVertexIDs );
+
+			int32 PolygonVertexNumberForEdgeVertices[ 2 ] = { INDEX_NONE, INDEX_NONE };
+			for( int32 PerimeterVertexNumber = 0; PerimeterVertexNumber < PolygonPerimeterVertexIDs.Num(); ++PerimeterVertexNumber )
+			{
+				const FVertexID PerimeterVertexID = PolygonPerimeterVertexIDs[ PerimeterVertexNumber ];
+				if( PerimeterVertexID == OriginalEdgeVertexIDs[ 0 ] )
+				{
+					// Okay, this polygon vertex is part of our edge (first edge vertex)
+					PolygonVertexNumberForEdgeVertices[ 0 ] = PerimeterVertexNumber;
+				}
+				else if( PerimeterVertexID == OriginalEdgeVertexIDs[ 1 ] )
+				{
+					// Okay, this polygon vertex is part of our edge (second edge vertex)
+					PolygonVertexNumberForEdgeVertices[ 1 ] = PerimeterVertexNumber;
+				}
+			}
+			check( PolygonVertexNumberForEdgeVertices[ 0 ] != INDEX_NONE && PolygonVertexNumberForEdgeVertices[ 1 ] != INDEX_NONE );
+
+			FAffectedPolygonEdge& AffectedPolygonEdge = *new( AffectedPolygonEdges ) FAffectedPolygonEdge();
+			AffectedPolygonEdge.PolygonID = PolygonID;
+			AffectedPolygonEdge.PolygonVertexNumbers[ 0 ] = PolygonVertexNumberForEdgeVertices[ 0 ];
+			AffectedPolygonEdge.PolygonVertexNumbers[ 1 ] = PolygonVertexNumberForEdgeVertices[ 1 ];
+
+			// @todo mesheditor perf: Not sure if we actually need to AddUnique() here -- can an edge be connected to the same polygon more than once? (tube shape). Seems sketchy if so?  And if so, GetEdgeConnectedPolygonCount() probably should have only reported 1 connection anyway
+			AffectedPolygons.AddUnique( AffectedPolygonEdge.PolygonID );
+		}
+	}
+
+	{
+		static TArray<FVerticesForEdge> VerticesForEdges;
+		VerticesForEdges.Reset();
+
+		// We'll keep the existing edge, but update it to connect to the first new vertex.  The second vertex of the edge will 
+		// now connect to the first (new) vertex ID, and so on.  The incoming vertices are expected to be ordered correctly.
+		{
+			FVerticesForEdge& VerticesForEdge = *new( VerticesForEdges ) FVerticesForEdge();
+
+			VerticesForEdge.EdgeID = EdgeID;
+			VerticesForEdge.NewVertexID0 = OriginalEdgeVertexIDs[ 0 ];
+			VerticesForEdge.NewVertexID1 = NewVertexIDs[ 0 ];
+		}
+
+		SetEdgesVertices( VerticesForEdges );
+	}
+
+	// Create new edges.  One for each of the new vertex positions passed in.
+	{
+		const int32 NewEdgeCount = NewVertexPositions.Num();
+
+		static TArray<FEdgeToCreate> EdgesToCreate;
+		EdgesToCreate.Reset();
+		EdgesToCreate.Reserve( NewEdgeCount );
+		for( int32 NewEdgeNumber = 0; NewEdgeNumber < NewEdgeCount; ++NewEdgeNumber )
+		{
+			FEdgeToCreate& EdgeToCreate = *new( EdgesToCreate ) FEdgeToCreate();
+
+			EdgeToCreate.VertexID0 = NewVertexIDs[ NewEdgeNumber ];
+			EdgeToCreate.VertexID1 = ( NewEdgeNumber == ( NewEdgeCount - 1 ) ) ? OriginalEdgeFarVertexID : NewVertexIDs[ NewEdgeNumber + 1 ];
+
+			EdgeToCreate.ConnectedPolygons = AffectedPolygons;
+
+			// Copy edge attributes over from original edge
+			EdgeToCreate.EdgeAttributes.Attributes = OriginalEdgeAttributes;
+		}
+
+		static TArray< FEdgeID > NewEdgeIDs;
+		NewEdgeIDs.Reset();
+		CreateEdges( EdgesToCreate, /* Out */ NewEdgeIDs );
+	}
+
+	// Update all affected polygons with their new vertices.  Also, we'll fill in polygon-specific vertex attributes (texture coordinates)
+	{
+		for( const FAffectedPolygonEdge& AffectedPolygonEdge : AffectedPolygonEdges )
+		{
+			const FPolygonID PolygonID = AffectedPolygonEdge.PolygonID;
+
+			// Figure out where to start inserting vertices into this polygon, as well as which order to insert them.
+			bool bWindsForward = AffectedPolygonEdge.PolygonVertexNumbers[ 1 ] > AffectedPolygonEdge.PolygonVertexNumbers[ 0 ];
+			const int32 LargerIndex = bWindsForward ? AffectedPolygonEdge.PolygonVertexNumbers[ 1 ] : AffectedPolygonEdge.PolygonVertexNumbers[ 0 ];
+			const bool bIsContiguous = FMath::Abs( (int32)AffectedPolygonEdge.PolygonVertexNumbers[ 1 ] - (int32)AffectedPolygonEdge.PolygonVertexNumbers[ 0 ] ) == 1;
+			if( !bIsContiguous )
+			{
+				bWindsForward = !bWindsForward;
+			}
+
+			const int32 InsertAtIndex = bIsContiguous ? LargerIndex : ( LargerIndex + 1 );
+
+			// Figure out what the polygon's NEW vertex count (after we insert new vertices) will be
+			const int32 NewPolygonPerimeterVertexCount = GetPolygonPerimeterVertexCount( PolygonID ) + NewVertexPositions.Num();
+
+			static TArray<FVertexAndAttributes> VerticesToInsert;
+			VerticesToInsert.Reset( NewVertexPositions.Num() );
+			VerticesToInsert.SetNum( NewVertexPositions.Num(), false );
+			for( int32 InsertVertexNumber = 0; InsertVertexNumber < NewVertexPositions.Num(); ++InsertVertexNumber )
+			{
+				FVertexAndAttributes& VertexToInsert = VerticesToInsert[ InsertVertexNumber ];
+
+				const int32 DirectionalVertexNumber = bWindsForward ? InsertVertexNumber : ( ( NewVertexPositions.Num() - 1 ) - InsertVertexNumber );
+				VertexToInsert.VertexID = NewVertexIDs[ InsertVertexNumber ];
+				VertexToInsert.VertexInstanceID = FVertexInstanceID::Invalid;
+
+				// Generate texture coordinates for this vertex by interpolating between the polygon vertex texture coordinates on this polygon edge
+				{
+					const float Split = SortedSplits[ DirectionalVertexNumber ];
+					for( int32 TextureCoordinateIndex = 0; TextureCoordinateIndex < TextureCoordinateCount; ++TextureCoordinateIndex )
+					{
+						const FVector4 NewPolygonVertexTextureCoordinate = FMath::Lerp(
+							GetPolygonPerimeterVertexAttribute( PolygonID, AffectedPolygonEdge.PolygonVertexNumbers[ 0 ], UEditableMeshAttribute::VertexTextureCoordinate(), TextureCoordinateIndex ),
+							GetPolygonPerimeterVertexAttribute( PolygonID, AffectedPolygonEdge.PolygonVertexNumbers[ 1 ], UEditableMeshAttribute::VertexTextureCoordinate(), TextureCoordinateIndex ),
+							Split );
+
+						VertexToInsert.PolygonVertexAttributes.Attributes.Add( FMeshElementAttributeData( UEditableMeshAttribute::VertexTextureCoordinate(), TextureCoordinateIndex, NewPolygonVertexTextureCoordinate ) );
+					}
+				}
+
+				// Generate vertex color for this vertex by interpolating between the vertex colors on the edge
+				{
+					const float Split = SortedSplits[ DirectionalVertexNumber ];
+					const FVector4 NewVertexColor = FMath::Lerp(
+						GetPolygonPerimeterVertexAttribute( PolygonID, AffectedPolygonEdge.PolygonVertexNumbers[ 0 ], UEditableMeshAttribute::VertexColor(), 0 ),
+						GetPolygonPerimeterVertexAttribute( PolygonID, AffectedPolygonEdge.PolygonVertexNumbers[ 1 ], UEditableMeshAttribute::VertexColor(), 0 ),
+						Split );
+
+					VertexToInsert.PolygonVertexAttributes.Attributes.Add( FMeshElementAttributeData( UEditableMeshAttribute::VertexColor(), 0, NewVertexColor ) );
+				}
+			}
+
+			// Add the new vertices to the polygon, and set texture coordinates for all of the affected polygon vertices
+			InsertPolygonPerimeterVertices( PolygonID, InsertAtIndex, VerticesToInsert );
+		}
+	}
+
+	// Generate normals and tangents
+	PolygonsPendingNewTangentBasis.Append( AffectedPolygons );
+
+	// Retriangulate all of the affected polygons
+	PolygonsPendingTriangulation.Append( AffectedPolygons );
+
+	UE_LOG( LogEditableMesh, Verbose, TEXT( "* SplitEdge returned %s" ), *LogHelpers::ArrayToString( OutNewVertexIDs ) );
+}
+
+
+void UEditableMesh::FindPolygonLoop( const FEdgeID EdgeID, TArray<FEdgeID>& OutEdgeLoopEdgeIDs, TArray<FEdgeID>& OutFlippedEdgeIDs, TArray<FEdgeID>& OutReversedEdgeIDPathToTake, TArray<FPolygonID>& OutPolygonIDsToSplit ) const
+{
+	OutEdgeLoopEdgeIDs.Reset();
+	OutFlippedEdgeIDs.Reset();
+	OutReversedEdgeIDPathToTake.Reset();
+	OutPolygonIDsToSplit.Reset();
+
+	// Is the edge we're starting on a border edge?
+	bool bStartedOnBorderEdge = ( GetEdgeConnectedPolygonCount( EdgeID ) <= 1 );
+
+	// We'll actually do two passes searching for edges.  The first time, we'll flow along the polygons looking for a border edge.
+	// If we find one, it means we won't really have a loop, but instead we're just splitting a series of connected polygons.  In that
+	// case, the polygon with the border edge will become the start of our search for an opposing border edge.  Otherwise, we'll start
+	// at the input polygon and flow around until we come back to that polygon again.  If we can't make it back, then no polygons
+	// will be split by this operation.
+	bool bIsSearchingForBorderEdge = !bStartedOnBorderEdge;
+
+	// Keep track of whether we actually looped back around to the starting edge (rather than simply splitting across a series
+	// of polygons that both end at border edges.)
+	bool bIsCompleteLoop = false;
+
+	FEdgeID CurrentEdgeID = EdgeID;
+	bool bCurrentEdgeIsBorderEdge = bStartedOnBorderEdge;
+	bool bCurrentEdgeIsInOppositeDirection = false;
+	bool bCurrentEdgeIsInOppositeDirectionFromStartEdge = false;
+	for( ; ; )
+	{
+		// Add the current edge!
+		checkSlow( !OutEdgeLoopEdgeIDs.Contains( CurrentEdgeID ) );
+		OutEdgeLoopEdgeIDs.Add( CurrentEdgeID );
+		if( bCurrentEdgeIsInOppositeDirection )
+		{
+			OutFlippedEdgeIDs.Add( CurrentEdgeID );
+		}
+
+		FVertexID CurrentEdgeVertexIDs[ 2 ];
+		GetEdgeVertices( CurrentEdgeID, /* Out */ CurrentEdgeVertexIDs[ 0 ], /* Out */ CurrentEdgeVertexIDs[ 1 ] );
+
+		const FVector CurrentEdgeVertex0 = GetVertexAttribute( CurrentEdgeVertexIDs[ 0 ], UEditableMeshAttribute::VertexPosition(), 0 );
+		const FVector CurrentEdgeVertex1 = GetVertexAttribute( CurrentEdgeVertexIDs[ 1 ], UEditableMeshAttribute::VertexPosition(), 0 );
+		const FVector CurrentEdgeDirection = ( CurrentEdgeVertex1 - CurrentEdgeVertex0 ).GetSafeNormal();
+
+		const FEdgeID NextEdgeIDInPath = OutReversedEdgeIDPathToTake.Num() > 0 ? OutReversedEdgeIDPathToTake.Pop( false ) : FEdgeID::Invalid;
+
+		FEdgeID BestEdgeID = FEdgeID::Invalid;
+		FPolygonID BestEdgeSplitsPolygon = FPolygonID::Invalid;
+		bool bBestEdgeIsInOppositeDirection = false;
+		bool bBestEdgeIsBorderEdge = false;
+		float LargestAbsDotProduct = -1.0f;
+
+		// Let's take a look at all of the polygons connected to this edge.  These will start our loop.
+		const int32 ConnectedPolygonCount = GetEdgeConnectedPolygonCount( CurrentEdgeID );
+		for( int32 ConnectedPolygonNumber = 0; ConnectedPolygonNumber < ConnectedPolygonCount; ++ConnectedPolygonNumber )
+		{
+			const FPolygonID ConnectedPolygonID = GetEdgeConnectedPolygon( CurrentEdgeID, ConnectedPolygonNumber );
+
+			// Don't bother looking at the last polygon that was added to our split list.  We never want to backtrack!
+			if( OutPolygonIDsToSplit.Num() == 0 || ConnectedPolygonID != OutPolygonIDsToSplit.Last() )
+			{
+				static TArray<FEdgeID> CandidateEdgeIDs;
+				CandidateEdgeIDs.Reset();
+				GetPolygonPerimeterEdges( ConnectedPolygonID, /* Out */ CandidateEdgeIDs );
+
+				// Which edge of the connected polygon will be at the other end of our split?
+				for( const FEdgeID CandidateEdgeID : CandidateEdgeIDs )
+				{
+					// Don't bother with the edge we just came from
+					if( CandidateEdgeID != CurrentEdgeID )
+					{
+						// If we need to follow a specific path, then do that
+						if( NextEdgeIDInPath == FEdgeID::Invalid || CandidateEdgeID == NextEdgeIDInPath )
+						{
+							FVertexID CandidateEdgeVertexIDs[ 2 ];
+							GetEdgeVertices( CandidateEdgeID, /* Out */ CandidateEdgeVertexIDs[ 0 ], /* Out */ CandidateEdgeVertexIDs[ 1 ] );
+
+							const bool bIsBorderEdge = ( GetEdgeConnectedPolygonCount( CandidateEdgeID ) == 1 );
+
+							const FVector CandidateEdgeVertex0 = GetVertexAttribute( CandidateEdgeVertexIDs[ 0 ], UEditableMeshAttribute::VertexPosition(), 0 );
+							const FVector CandidateEdgeVertex1 = GetVertexAttribute( CandidateEdgeVertexIDs[ 1 ], UEditableMeshAttribute::VertexPosition(), 0 );
+							const FVector CandidateEdgeDirection = ( CandidateEdgeVertex1 - CandidateEdgeVertex0 ).GetSafeNormal();
+
+							const float DotProduct = FVector::DotProduct( CurrentEdgeDirection, CandidateEdgeDirection );
+							const float AbsDotProduct = FMath::Abs( DotProduct );
+
+							const float SameEdgeDirectionDotEpsilon = 0.05f;	// @todo mesheditor tweak
+							if( FMath::IsNearlyEqual( AbsDotProduct, LargestAbsDotProduct, SameEdgeDirectionDotEpsilon ) )
+							{
+								// If the candidate edge directions are pretty much the same, we'll choose the edge that flows closest to the
+								// direction that we split the last polygon in
+								if( OutEdgeLoopEdgeIDs.Num() > 1 )
+								{
+									FVertexID LastSplitEdgeVertexIDs[ 2 ];
+									GetEdgeVertices( OutEdgeLoopEdgeIDs[ OutEdgeLoopEdgeIDs.Num() - 2 ], /* Out */ LastSplitEdgeVertexIDs[ 0 ], /* Out */ LastSplitEdgeVertexIDs[ 1 ] );
+
+									const FVector LastSplitEdgeVertex0 = GetVertexAttribute( LastSplitEdgeVertexIDs[ 0 ], UEditableMeshAttribute::VertexPosition(), 0 );
+									const FVector LastSplitEdgeVertex1 = GetVertexAttribute( LastSplitEdgeVertexIDs[ 1 ], UEditableMeshAttribute::VertexPosition(), 0 );
+
+									const FVector DirectionTowardCenterOfCurrentEdge =
+										( FMath::Lerp( CurrentEdgeVertex0, CurrentEdgeVertex1, 0.5f ) -
+											FMath::Lerp( LastSplitEdgeVertex0, LastSplitEdgeVertex1, 0.5f ) ).GetSafeNormal();
+
+									const FVector DirectionTowardCenterOfCandidateEdge =
+										( FMath::Lerp( CandidateEdgeVertex0, CandidateEdgeVertex1, 0.5f ) -
+											FMath::Lerp( CurrentEdgeVertex0, CurrentEdgeVertex1, 0.5f ) ).GetSafeNormal();
+									const float CandidateEdgeDot = FVector::DotProduct( DirectionTowardCenterOfCurrentEdge, DirectionTowardCenterOfCandidateEdge );
+
+									check( BestEdgeID != FEdgeID::Invalid );
+
+									FVertexID BestEdgeVertexIDs[ 2 ];
+									GetEdgeVertices( BestEdgeID, /* Out */ BestEdgeVertexIDs[ 0 ], /* Out */ BestEdgeVertexIDs[ 1 ] );
+
+									const FVector BestEdgeVertex0 = GetVertexAttribute( BestEdgeVertexIDs[ 0 ], UEditableMeshAttribute::VertexPosition(), 0 );
+									const FVector BestEdgeVertex1 = GetVertexAttribute( BestEdgeVertexIDs[ 1 ], UEditableMeshAttribute::VertexPosition(), 0 );
+
+									const FVector DirectionTowardCenterOfBestEdge =
+										( FMath::Lerp( BestEdgeVertex0, BestEdgeVertex1, 0.5f ) -
+											FMath::Lerp( CurrentEdgeVertex0, CurrentEdgeVertex1, 0.5f ) ).GetSafeNormal();
+
+									const float BestEdgeDot = FVector::DotProduct( DirectionTowardCenterOfCurrentEdge, DirectionTowardCenterOfBestEdge );
+
+									if( CandidateEdgeDot > BestEdgeDot )
+									{
+										BestEdgeID = CandidateEdgeID;
+										BestEdgeSplitsPolygon = ConnectedPolygonID;
+										bBestEdgeIsInOppositeDirection = ( DotProduct < 0.0f );
+										bBestEdgeIsBorderEdge = bIsBorderEdge;
+										LargestAbsDotProduct = AbsDotProduct;
+									}
+								}
+								else
+								{
+									// Edge directions are the same, but this is the very first split so we don't have a "flow" direction yet.
+									// Go ahead and prefer the edge that is closer to the initial edge.  This helps in the (uncommon) case of 
+									// multiple colinear edges on the same polygon (such as after SplitEdge() is called to insert a vertex on a polygon.)
+									const float BestEdgeDistance = [this, BestEdgeID, CurrentEdgeVertex0, CurrentEdgeVertex1]() -> float
+									{
+										check( BestEdgeID != FEdgeID::Invalid );
+
+										FVertexID BestEdgeVertexIDs[ 2 ];
+										GetEdgeVertices( BestEdgeID, /* Out */ BestEdgeVertexIDs[ 0 ], /* Out */ BestEdgeVertexIDs[ 1 ] );
+
+										const FVector BestEdgeVertex0 = GetVertexAttribute( BestEdgeVertexIDs[ 0 ], UEditableMeshAttribute::VertexPosition(), 0 );
+										const FVector BestEdgeVertex1 = GetVertexAttribute( BestEdgeVertexIDs[ 1 ], UEditableMeshAttribute::VertexPosition(), 0 );
+
+										FVector ClosestPoint0, ClosestPoint1;
+										FMath::SegmentDistToSegmentSafe( CurrentEdgeVertex0, CurrentEdgeVertex1, BestEdgeVertex0, BestEdgeVertex1, /* Out */ ClosestPoint0, /* Out */ ClosestPoint1 );
+										return ( ClosestPoint1 - ClosestPoint0 ).Size();
+									}( );
+
+									const float CandidateEdgeDistance = [CurrentEdgeVertex0, CurrentEdgeVertex1, CandidateEdgeVertex0, CandidateEdgeVertex1]() -> float
+									{
+										FVector ClosestPoint0, ClosestPoint1;
+										FMath::SegmentDistToSegmentSafe( CurrentEdgeVertex0, CurrentEdgeVertex1, CandidateEdgeVertex0, CandidateEdgeVertex1, /* Out */ ClosestPoint0, /* Out */ ClosestPoint1 );
+										return ( ClosestPoint1 - ClosestPoint0 ).Size();
+									}( );
+
+									if( CandidateEdgeDistance < BestEdgeDistance )
+									{
+										BestEdgeID = CandidateEdgeID;
+										BestEdgeSplitsPolygon = ConnectedPolygonID;
+										bBestEdgeIsInOppositeDirection = ( DotProduct < 0.0f );
+										bBestEdgeIsBorderEdge = bIsBorderEdge;
+										LargestAbsDotProduct = AbsDotProduct;
+									}
+								}
+							}
+							else if( AbsDotProduct > LargestAbsDotProduct )
+							{
+								// This edge angle is the closest to our current edge so far!
+								BestEdgeID = CandidateEdgeID;
+								BestEdgeSplitsPolygon = ConnectedPolygonID;
+								bBestEdgeIsInOppositeDirection = ( DotProduct < 0.0f );
+								bBestEdgeIsBorderEdge = bIsBorderEdge;
+								LargestAbsDotProduct = AbsDotProduct;
+							}
+						}
+					}
+				}
+			}
+		}
+
+		if( BestEdgeID != FEdgeID::Invalid &&
+			!OutPolygonIDsToSplit.Contains( BestEdgeSplitsPolygon ) )	// If we try to re-split the same polygon twice, then this loop is not valid
+		{
+			// OK, this polygon will definitely be split
+			OutPolygonIDsToSplit.Add( BestEdgeSplitsPolygon );
+
+			CurrentEdgeID = BestEdgeID;
+			bCurrentEdgeIsBorderEdge = bBestEdgeIsBorderEdge;
+			bCurrentEdgeIsInOppositeDirection = bBestEdgeIsInOppositeDirection;
+			if( bBestEdgeIsInOppositeDirection )
+			{
+				bCurrentEdgeIsInOppositeDirectionFromStartEdge = !bCurrentEdgeIsInOppositeDirectionFromStartEdge;
+			}
+
+			// Is the best edge already part of our loop?  If so, then we're done!
+			if( OutEdgeLoopEdgeIDs[ 0 ] == BestEdgeID )
+			{
+				bIsCompleteLoop = true;
+				break;
+			}
+			else if( OutEdgeLoopEdgeIDs.Contains( BestEdgeID ) )	// @todo mesheditor perf: Might want to use a TSet here, but we'll still need to keep a list too (order is important.)
+			{
+				// We ended up back at an edge that we already split, but it wasn't the edge that we started on.  This is
+				// not a valid loop, so clear our path and bail out.
+
+				// @todo mesheditor edgeloop: We need to revisit how we handle border edges and non-manifold edges.  Basically instead of searching
+				// for border edges, we should try looping around the "back" of polygons we've already visited, trying to get back to the starting
+				// edge.  This will allow the case of a cube mesh with a single non-manifold face to be split properly.  We just need to make sure
+				// we're not trying to split the same edge multiple times in different locations.  Probably, we don't want to split the same polygon
+				// more than once either.
+				OutEdgeLoopEdgeIDs.Reset();
+				OutFlippedEdgeIDs.Reset();
+				OutPolygonIDsToSplit.Reset();
+
+				break;
+			}
+			else if( bBestEdgeIsBorderEdge && bIsSearchingForBorderEdge )
+			{
+				// We found a border edge, so stop the search.  We'll now start over at this edge to form our loop.
+				bStartedOnBorderEdge = true;
+
+				bIsSearchingForBorderEdge = false;
+				bCurrentEdgeIsInOppositeDirection = bCurrentEdgeIsInOppositeDirectionFromStartEdge;
+
+				// Follow the path we took to get here, in reverse order, to make sure we get back to the edge we 
+				// were asked to create a loop on
+				OutReversedEdgeIDPathToTake = OutEdgeLoopEdgeIDs;
+
+				bIsCompleteLoop = false;
+
+				OutEdgeLoopEdgeIDs.Reset();
+				OutFlippedEdgeIDs.Reset();
+				OutPolygonIDsToSplit.Reset();
+			}
+			else
+			{
+				// Proceed to the next edge and try to continue the loop.  If we're at a border edge, the loop will definitely end here.
+			}
+		}
+		else
+		{
+			if( bStartedOnBorderEdge && bCurrentEdgeIsBorderEdge )
+			{
+				// We started on a border edge, and we've found the border edge on the other side of the polygons we'll be splitting.
+				// This isn't actually a loop, but we'll still split the polygons.
+			}
+			else
+			{
+				// We couldn't even find another edge, so we're done.
+				OutEdgeLoopEdgeIDs.Reset();
+				OutFlippedEdgeIDs.Reset();
+				OutPolygonIDsToSplit.Reset();
+			}
+
+			break;
+		}
+	}
+
+	// We're always splitting the same number of polygons as we have edges in the loop (these can be zero), except
+	// in the border edge case, where we're always splitting one less polygon.
+	if( bStartedOnBorderEdge &&
+		!bIsCompleteLoop )	// It's possible to start and finish on the same border edge, which is a complete loop
+	{
+		// We're splitting a series of polygons between two border edges
+		check( ( OutEdgeLoopEdgeIDs.Num() == 0 && OutPolygonIDsToSplit.Num() == 0 ) || ( OutEdgeLoopEdgeIDs.Num() == ( OutPolygonIDsToSplit.Num() + 1 ) ) );
+	}
+	else
+	{
+		// We're splitting polygons spanning a full loop of edges.  The starting edge is the same as the beginning edge.
+		check( OutEdgeLoopEdgeIDs.Num() == OutPolygonIDsToSplit.Num() );
+	}
+}
+
+
+void UEditableMesh::InsertEdgeLoop( const FEdgeID EdgeID, const TArray<float>& Splits, TArray<FEdgeID>& OutNewEdgeIDs )
+{
+	OutNewEdgeIDs.Reset();
+
+	// NOTE: The incoming splits should always be between 0.0 and 1.0, representing progress along 
+	//       the edge from the edge's first vertex toward it's other vertex.  The order doesn't matter (we'll sort them.)
+
+	// @todo mesheditor: Test with concave polygons -- probably we need to disallow/avoid splits that cross a concave part
+
+
+	static TArray<FEdgeID> EdgeLoopEdgeIDs;
+	EdgeLoopEdgeIDs.Reset();
+
+	static TArray<FEdgeID> FlippedEdgeIDs;
+	FlippedEdgeIDs.Reset();
+
+	static TArray<FEdgeID> ReversedEdgeIDPathToTake;
+	ReversedEdgeIDPathToTake.Reset();
+
+	static TArray<FPolygonID> PolygonIDsToSplit;
+	PolygonIDsToSplit.Reset();
+
+
+	FindPolygonLoop( 
+		EdgeID, 
+		/* Out */ EdgeLoopEdgeIDs,
+		/* Out */ FlippedEdgeIDs,
+		/* Out */ ReversedEdgeIDPathToTake,
+		/* Out */ PolygonIDsToSplit );
+
+	static TSet<FEdgeID> FlippedEdgeIDSet;
+	FlippedEdgeIDSet.Reset();
+	FlippedEdgeIDSet.Append( FlippedEdgeIDs );
+
+	check( Splits.Num() > 0 );
+
+	// Sort the split values smallest to largest.  We'll be adding a strip of vertices for each split, and the
+	// IDs for those new vertices need to be in order.
+	static TArray<float> SortedSplits;
+	SortedSplits.Reset();
+	SortedSplits = Splits;
+	if( SortedSplits.Num() > 1 )
+	{
+		SortedSplits.Sort();
+	}
+
+
+	if( PolygonIDsToSplit.Num() > 0 )
+	{
+		// Keep track of the new vertices create by splitting all of the edges.  For each edge we split, an array of
+		// vertex IDs for each split along that edge
+		static TArray<TArray<FVertexID>> NewVertexIDsForEachEdge;
+		NewVertexIDsForEachEdge.Reset();
+
+
+		// Now let's go through and create new vertices for the loops by splitting edges
+		{
+			for( int32 EdgeLoopEdgeNumber = 0; EdgeLoopEdgeNumber < EdgeLoopEdgeIDs.Num(); ++EdgeLoopEdgeNumber )
+			{
+				const FEdgeID EdgeLoopEdgeID = EdgeLoopEdgeIDs[ EdgeLoopEdgeNumber ];
+
+				// If the edge winds in the opposite direction from the last edge, we'll need to flip the split positions around
+				const bool bIsFlipped = FlippedEdgeIDSet.Contains( EdgeLoopEdgeID );
+				if( bIsFlipped )
+				{
+					static TArray<float> TempSplits;
+					TempSplits.SetNumUninitialized( SortedSplits.Num(), false );
+					for( int32 SplitIndex = 0; SplitIndex < SortedSplits.Num(); ++SplitIndex )
+					{
+						TempSplits[ SplitIndex ] = 1.0f - SortedSplits[ ( SortedSplits.Num() - 1 ) - SplitIndex ];
+					}
+					SortedSplits = TempSplits;
+				}
+
+				// Split this edge
+				static TArray<FVertexID> CurrentVertexIDs;
+				CurrentVertexIDs.Reset();
+				SplitEdge( EdgeLoopEdgeID, SortedSplits, /* Out */ CurrentVertexIDs );
+
+				// If the edge winding is backwards, we'll reverse the order of the vertex IDs in our list
+				// @todo mesheditor urgent edgeloop: INCORRECT (bIsFlipped is supposed to swap from last time -- TempSplits.)  Also not sure if needed or not... might be double compensating here.
+				static TArray<FVertexID> NewVertexIDsForEdge;
+				NewVertexIDsForEdge.SetNum( CurrentVertexIDs.Num(), false );
+				for( int32 VertexNumber = 0; VertexNumber < CurrentVertexIDs.Num(); ++VertexNumber )
+				{
+					NewVertexIDsForEdge[ /*bIsFlipped ? */( ( CurrentVertexIDs.Num() - VertexNumber ) - 1 ) /*: VertexNumber */] = CurrentVertexIDs[ VertexNumber ];
+				}
+				NewVertexIDsForEachEdge.Add( NewVertexIDsForEdge );
+			}
+		}
+
+
+		// Time to create new polygons for the split faces (and delete the old ones)
+		{
+			static TArray<FPolygonToSplit> PolygonsToSplit;
+			PolygonsToSplit.Reset();
+
+			for( int32 PolygonToSplitIter = 0; PolygonToSplitIter < PolygonIDsToSplit.Num(); ++PolygonToSplitIter )
+			{
+				const FPolygonID PolygonID = PolygonIDsToSplit[ PolygonToSplitIter ];
+
+				FPolygonToSplit& PolygonToSplit = *new( PolygonsToSplit ) FPolygonToSplit();
+				PolygonToSplit.PolygonID = PolygonID;
+
+				// The first and second edges connected to this polygon that are being split up
+				const int32 FirstEdgeNumber = PolygonToSplitIter;
+				const int32 SecondEdgeNumber = ( PolygonToSplitIter + 1 ) % EdgeLoopEdgeIDs.Num();
+
+				const FEdgeID FirstSplitEdgeID = EdgeLoopEdgeIDs[ FirstEdgeNumber ];
+				const FEdgeID SecondSplitEdgeID = EdgeLoopEdgeIDs[ SecondEdgeNumber ];
+				check( FirstSplitEdgeID != SecondSplitEdgeID );
+
+				// The (ordered) list of new vertices that was created by splitting the first and second edge.  One for each split.
+				const TArray<FVertexID>& FirstSplitEdgeNewVertexIDs = NewVertexIDsForEachEdge[ FirstEdgeNumber ];
+				const TArray<FVertexID>& SecondSplitEdgeNewVertexIDs = NewVertexIDsForEachEdge[ SecondEdgeNumber ];
+
+				for( int32 SplitIter = 0; SplitIter < SortedSplits.Num(); ++SplitIter )
+				{
+					FVertexPair& NewVertexPair = *new( PolygonToSplit.VertexPairsToSplitAt ) FVertexPair();
+					NewVertexPair.VertexID0 = FirstSplitEdgeNewVertexIDs[ SplitIter ];
+					NewVertexPair.VertexID1 = SecondSplitEdgeNewVertexIDs[ SplitIter ];
+				}
+			}
+
+
+			// Actually split up the polygons
+			static TArray<FEdgeID> NewEdgeIDs;
+			NewEdgeIDs.Reset();
+			SplitPolygons( PolygonsToSplit, /* Out */ NewEdgeIDs );
+
+			OutNewEdgeIDs.Append( NewEdgeIDs );
+		}
@@ -1960,736 +3656,6 @@
 }
 
 
-void UEditableMesh::CreateMissingPolygonHoleEdges( const FPolygonRef PolygonRef, const int32 HoleNumber, TArray<FEdgeID>& OutNewEdgeIDs )
-{
-	OutNewEdgeIDs.Reset();
-
-	const int32 NumHoleEdges = GetPolygonHoleEdgeCount( PolygonRef, HoleNumber );
-	const int32 NumHoleVertices = NumHoleEdges;		// Edge and vertex count are always the same
-
-	for( int32 HoleEdgeNumber = 0; HoleEdgeNumber < NumHoleEdges; ++HoleEdgeNumber )
-	{
-		const int32 HoleVertexNumber = HoleEdgeNumber;	// Edge and vertex counts are always the same
-
-		const FVertexID VertexID = GetPolygonHoleVertex( PolygonRef, HoleNumber, HoleVertexNumber );
-		const FVertexID NextVertexID = GetPolygonHoleVertex( PolygonRef, HoleNumber, ( HoleVertexNumber + 1 ) % NumHoleVertices );
-
-		// Find the edge that connects these vertices
-		FEdgeID FoundEdgeID = FEdgeID::Invalid;
-		{
-			bool bFoundEdge = false;
-
-			const int32 NumVertexConnectedEdges = GetVertexConnectedEdgeCount( VertexID );	// @todo mesheditor perf: Could be made faster by always iterating over the vertex with the fewest number of connected edges, instead of the first edge vertex
-			for( int32 VertexEdgeNumber = 0; VertexEdgeNumber < NumVertexConnectedEdges; ++VertexEdgeNumber )
-			{
-				const FEdgeID VertexConnectedEdgeID = GetVertexConnectedEdge( VertexID, VertexEdgeNumber );
-
-				// Try the edge's first vertex.  Does it point to us?
-				FVertexID OtherEdgeVertexID = GetEdgeVertex( VertexConnectedEdgeID, 0 );
-				if( OtherEdgeVertexID == VertexID )
-				{
-					// Must be the the other one
-					OtherEdgeVertexID = GetEdgeVertex( VertexConnectedEdgeID, 1 );
-				}
-
-				if( OtherEdgeVertexID == NextVertexID )
-				{
-					// We found the edge!
-					FoundEdgeID = VertexConnectedEdgeID;
-					bFoundEdge = true;
-					break;
-				}
-			}
-
-			if( !bFoundEdge )
-			{
-				// Create the new edge!  Note that this does not connect the edge to the polygon.  We expect the caller to do that afterwards.
-				FEdgeToCreate EdgeToCreate;
-				EdgeToCreate.VertexID0 = VertexID;
-				EdgeToCreate.VertexID1 = NextVertexID;
-
-				static TArray< FEdgeToCreate > EdgesToCreate;
-				EdgesToCreate.Reset();
-				EdgesToCreate.Add( EdgeToCreate );
-
-				static TArray< FEdgeID > NewEdgeIDs;
-				NewEdgeIDs.Reset();
-				CreateEdges( EdgesToCreate, /* Out */ NewEdgeIDs );
-
-				OutNewEdgeIDs.Append( NewEdgeIDs );
-			}
-		}
-	}
-}
-
-
-void UEditableMesh::SplitEdge( const FEdgeID EdgeID, const TArray<float>& Splits, TArray<FVertexID>& OutNewVertexIDs )
-{
-	// NOTE: The incoming splits should always be between 0.0 and 1.0, representing progress along 
-	//       the edge from the edge's first vertex toward it's other vertex.  The order doesn't matter (we'll sort them.)
-
-	check( Splits.Num() > 0 );
-
-	// Sort the split values smallest to largest.  We'll be adding a strip of vertices for each split, and the
-	// indices for those new vertices need to be in order.
-	static TArray<float> SortedSplits;
-	SortedSplits.Reset();
-	SortedSplits = Splits;
-	if( SortedSplits.Num() > 1 )
-	{
-		SortedSplits.Sort();
-	}
-
-	FVertexID OriginalEdgeVertexIDs[ 2 ];
-	GetEdgeVertices( EdgeID, /* Out */ OriginalEdgeVertexIDs[0], /* Out */ OriginalEdgeVertexIDs[1] );
-
-	static TArray<FVector> NewVertexPositions;
-	NewVertexPositions.Reset();
-	NewVertexPositions.AddUninitialized( SortedSplits.Num() );
-
-	for( int32 NewVertexNumber = 0; NewVertexNumber < NewVertexPositions.Num(); ++NewVertexNumber )
-	{
-		const float Split = SortedSplits[ NewVertexNumber ];
-		check( Split >= 0.0f && Split <= 1.0f );
-
-		NewVertexPositions[ NewVertexNumber ] = FMath::Lerp(
-			FVector( GetVertexAttribute( OriginalEdgeVertexIDs[ 0 ], UEditableMeshAttribute::VertexPosition(), 0 ) ),
-			FVector( GetVertexAttribute( OriginalEdgeVertexIDs[ 1 ], UEditableMeshAttribute::VertexPosition(), 0 ) ),
-			Split );
-	}
-
-	// Split the edge, and connect the vertex to the polygons that share the two new edges
-	const FVertexID OriginalEdgeFarVertexID = OriginalEdgeVertexIDs[ 1 ];
-
-	static TArray<FMeshElementAttributeData> OriginalEdgeAttributes;
-	OriginalEdgeAttributes.Reset();
-	{
-		for( const FName AttributeName : UEditableMesh::GetValidEdgeAttributes() )
-		{
-			const int32 MaxAttributeIndex = GetMaxAttributeIndex( AttributeName );
-			for( int32 AttributeIndex = 0; AttributeIndex < MaxAttributeIndex; ++AttributeIndex )
-			{
-				FMeshElementAttributeData& EdgeAttribute = *new( OriginalEdgeAttributes ) FMeshElementAttributeData(
-					AttributeName,
-					AttributeIndex,
-					GetEdgeAttribute( EdgeID, AttributeName, AttributeIndex ) );
-			}
-		}
-	}
-
-
-	// Create new vertices.  Their texture coordinates will be undefined, for now.  We'll set that later by using the connected 
-	// polygon's vertex data.  Also, the vertex's normals and tangents will be automatically generated.
-	static TArray<FVertexID> NewVertexIDs;
-	NewVertexIDs.Reset();
-	{
-		OutNewVertexIDs.Reset();
-		OutNewVertexIDs.Reserve( NewVertexPositions.Num() );
-
-		static TArray<FVertexToCreate> VerticesToCreate;
-		VerticesToCreate.Reset();
-		VerticesToCreate.Reserve( NewVertexPositions.Num() );
-
-		for( int32 NewVertexNumber = 0; NewVertexNumber < NewVertexPositions.Num(); ++NewVertexNumber )
-		{
-			const FVector& NewVertexPosition = NewVertexPositions[ NewVertexNumber ];
-
-			FVertexToCreate& VertexToCreate = *new( VerticesToCreate )FVertexToCreate();
-			VertexToCreate.VertexAttributes.Attributes.Add( FMeshElementAttributeData( UEditableMeshAttribute::VertexPosition(), 0, NewVertexPosition ) );
-		}
-
-		CreateVertices( VerticesToCreate, /* Out */ NewVertexIDs );
-
-		OutNewVertexIDs.Append( NewVertexIDs );
-	}
-
-	// Figure out which polygons are connected to the original edge.  Those polygons are going to need their
-	// connected vertices updated later on.
-	struct FAffectedPolygonEdge
-	{
-		FPolygonRef PolygonRef;
-		int32 PolygonVertexNumbers[ 2 ];
-	};
-	static TArray<FAffectedPolygonEdge> AffectedPolygonEdges;
-	AffectedPolygonEdges.Reset();
-
-	static TArray<FPolygonRef> AffectedPolygons;
-	AffectedPolygons.Reset();
-
-	{
-		const int32 ConnectedPolygonCount = GetEdgeConnectedPolygonCount( EdgeID );
-		for( int32 ConnectedPolygonNumber = 0; ConnectedPolygonNumber < ConnectedPolygonCount; ++ConnectedPolygonNumber )
-		{
-			const FPolygonRef PolygonRef = GetEdgeConnectedPolygon( EdgeID, ConnectedPolygonNumber );
-
-			// Figure out which of this polygon's vertices are part of the edge we're adding a vertex to
-			// @todo mesheditor hole: Add support for polygon holes, too!
-			static TArray<FVertexID> PolygonPerimeterVertexIDs;
-			PolygonPerimeterVertexIDs.Reset();
-			GetPolygonPerimeterVertices( PolygonRef, /* Out */ PolygonPerimeterVertexIDs );
-
-			int32 PolygonVertexNumberForEdgeVertices[ 2 ] = { INDEX_NONE, INDEX_NONE };
-			for( int32 PerimeterVertexNumber = 0; PerimeterVertexNumber < PolygonPerimeterVertexIDs.Num(); ++PerimeterVertexNumber )
-			{
-				const FVertexID PerimeterVertexID = PolygonPerimeterVertexIDs[ PerimeterVertexNumber ];
-				if( PerimeterVertexID == OriginalEdgeVertexIDs[ 0 ] )
-				{
-					// Okay, this polygon vertex is part of our edge (first edge vertex)
-					PolygonVertexNumberForEdgeVertices[ 0 ] = PerimeterVertexNumber;
-				}
-				else if( PerimeterVertexID == OriginalEdgeVertexIDs[ 1 ] )
-				{
-					// Okay, this polygon vertex is part of our edge (second edge vertex)
-					PolygonVertexNumberForEdgeVertices[ 1 ] = PerimeterVertexNumber;
-				}
-			}
-			check( PolygonVertexNumberForEdgeVertices[ 0 ] != INDEX_NONE && PolygonVertexNumberForEdgeVertices[ 1 ] != INDEX_NONE );
-
-			FAffectedPolygonEdge& AffectedPolygonEdge = *new( AffectedPolygonEdges ) FAffectedPolygonEdge();
-			AffectedPolygonEdge.PolygonRef = PolygonRef;
-			AffectedPolygonEdge.PolygonVertexNumbers[ 0 ] = PolygonVertexNumberForEdgeVertices[ 0 ];
-			AffectedPolygonEdge.PolygonVertexNumbers[ 1 ] = PolygonVertexNumberForEdgeVertices[ 1 ];
-
-			// @todo mesheditor perf: Not sure if we actually need to AddUnique() here -- can an edge be connected to the same polygon more than once? (tube shape). Seems sketchy if so?  And if so, GetEdgeConnectedPolygonCount() probably should have only reported 1 connection anyway
-			AffectedPolygons.AddUnique( AffectedPolygonEdge.PolygonRef );
-		}
-	}
-
-
-	// When undoing this change, we'll need to retriangulate before the vertices are destroyed (after CreateVertices() is called above),
-	// so that the vertices will be fully orphaned first (otherwise, the rendering vertices would still be referenced.)
-	{
-		const bool bOnlyOnUndo = true;
-		RetriangulatePolygons( AffectedPolygons, bOnlyOnUndo );
-	}
-
-
-	{
-		static TArray<FVerticesForEdge> VerticesForEdges;
-		VerticesForEdges.Reset();
-
-		// We'll keep the existing edge, but update it to connect to the first new vertex.  The second vertex of the edge will 
-		// now connect to the first (new) vertex ID, and so on.  The incoming vertices are expected to be ordered correctly.
-		{
-			FVerticesForEdge& VerticesForEdge = *new( VerticesForEdges ) FVerticesForEdge();
-
-			VerticesForEdge.EdgeID = EdgeID;
-			VerticesForEdge.NewVertexID0 = OriginalEdgeVertexIDs[ 0 ];
-			VerticesForEdge.NewVertexID1 = NewVertexIDs[ 0 ];
-		}
-
-		SetEdgesVertices( VerticesForEdges );
-	}
-
-	// Create new edges.  One for each of the new vertex positions passed in.
-	{
-		const int32 NewEdgeCount = NewVertexPositions.Num();
-
-		static TArray<FEdgeToCreate> EdgesToCreate;
-		EdgesToCreate.Reset();
-		EdgesToCreate.Reserve( NewEdgeCount );
-		for( int32 NewEdgeNumber = 0; NewEdgeNumber < NewEdgeCount; ++NewEdgeNumber )
-		{
-			FEdgeToCreate& EdgeToCreate = *new( EdgesToCreate ) FEdgeToCreate();
-
-			EdgeToCreate.VertexID0 = NewVertexIDs[ NewEdgeNumber ];
-			EdgeToCreate.VertexID1 = ( NewEdgeNumber == ( NewEdgeCount - 1 ) ) ? OriginalEdgeFarVertexID : NewVertexIDs[ NewEdgeNumber + 1 ];
-
-			EdgeToCreate.ConnectedPolygons = AffectedPolygons;
-
-			// Copy edge attributes over from original edge
-			EdgeToCreate.EdgeAttributes.Attributes = OriginalEdgeAttributes;
-		}
-
-		static TArray< FEdgeID > NewEdgeIDs;
-		NewEdgeIDs.Reset();
-		CreateEdges( EdgesToCreate, /* Out */ NewEdgeIDs );
-	}
-
-	// Update all affected polygons with their new vertices.  Also, we'll fill in polygon-specific vertex attributes (texture coordinates)
-	{
-		for( const FAffectedPolygonEdge& AffectedPolygonEdge : AffectedPolygonEdges )
-		{
-			const FPolygonRef PolygonRef = AffectedPolygonEdge.PolygonRef;
-
-			// Figure out where to start inserting vertices into this polygon, as well as which order to insert them.
-			bool bWindsForward = AffectedPolygonEdge.PolygonVertexNumbers[ 1 ] > AffectedPolygonEdge.PolygonVertexNumbers[ 0 ];
-			const int32 LargerIndex = bWindsForward ? AffectedPolygonEdge.PolygonVertexNumbers[ 1 ] : AffectedPolygonEdge.PolygonVertexNumbers[ 0 ];
-			const bool bIsContiguous = FMath::Abs( (int32)AffectedPolygonEdge.PolygonVertexNumbers[ 1 ] - (int32)AffectedPolygonEdge.PolygonVertexNumbers[ 0 ] ) == 1;
-			if( !bIsContiguous )
-			{
-				bWindsForward = !bWindsForward;
-			}
-
-			const int32 InsertAtIndex = bIsContiguous ? LargerIndex : ( LargerIndex + 1 );
-
-			// Figure out what the polygon's NEW vertex count (after we insert new vertices) will be
-			const int32 NewPolygonPerimeterVertexCount = GetPolygonPerimeterVertexCount( PolygonRef ) + NewVertexPositions.Num();
-
-			static TArray<FVertexAndAttributes> VerticesToInsert;
-			VerticesToInsert.Reset( NewVertexPositions.Num() );
-			VerticesToInsert.SetNum( NewVertexPositions.Num(), false );
-			for( int32 InsertVertexNumber = 0; InsertVertexNumber < NewVertexPositions.Num(); ++InsertVertexNumber )
-			{
-				FVertexAndAttributes& VertexToInsert = VerticesToInsert[ InsertVertexNumber ];
-
-				const int32 DirectionalVertexNumber = bWindsForward ? InsertVertexNumber : ( ( NewVertexPositions.Num() - 1 ) - InsertVertexNumber );
-				VertexToInsert.VertexID = NewVertexIDs[ InsertVertexNumber ];
-
-				// Generate texture coordinates for this vertex by interpolating between the polygon vertex texture coordinates on this polygon edge
-				{
-					const float Split = SortedSplits[ DirectionalVertexNumber ];
-					for( int32 TextureCoordinateIndex = 0; TextureCoordinateIndex < TextureCoordinateCount; ++TextureCoordinateIndex )
-					{
-						const FVector4 NewPolygonVertexTextureCoordinate = FMath::Lerp(
-							GetPolygonPerimeterVertexAttribute( PolygonRef, AffectedPolygonEdge.PolygonVertexNumbers[ 0 ], UEditableMeshAttribute::VertexTextureCoordinate(), TextureCoordinateIndex ),
-							GetPolygonPerimeterVertexAttribute( PolygonRef, AffectedPolygonEdge.PolygonVertexNumbers[ 1 ], UEditableMeshAttribute::VertexTextureCoordinate(), TextureCoordinateIndex ),
-							Split );
-
-						VertexToInsert.PolygonVertexAttributes.Attributes.Add( FMeshElementAttributeData( UEditableMeshAttribute::VertexTextureCoordinate(), TextureCoordinateIndex, NewPolygonVertexTextureCoordinate ) );
-					}
-				}
-
-				// Generate vertex color for this vertex by interpolating between the vertex colors on the edge
-				{
-					const float Split = SortedSplits[ DirectionalVertexNumber ];
-					const FVector4 NewVertexColor = FMath::Lerp(
-						GetPolygonPerimeterVertexAttribute( PolygonRef, AffectedPolygonEdge.PolygonVertexNumbers[ 0 ], UEditableMeshAttribute::VertexColor(), 0 ),
-						GetPolygonPerimeterVertexAttribute( PolygonRef, AffectedPolygonEdge.PolygonVertexNumbers[ 1 ], UEditableMeshAttribute::VertexColor(), 0 ),
-						Split );
-
-					VertexToInsert.PolygonVertexAttributes.Attributes.Add( FMeshElementAttributeData( UEditableMeshAttribute::VertexColor(), 0, NewVertexColor ) );
-				}
-			}
-
-			// Add the new vertices to the polygon, and set texture coordinates for all of the affected polygon vertices
-			InsertPolygonPerimeterVertices( PolygonRef, InsertAtIndex, VerticesToInsert );
-		}
-	}
-
-	// Generate normals and tangents
-	GenerateNormalsAndTangentsForPolygons( AffectedPolygons );
-
-	// Retriangulate all of the affected polygons
-	{
-		const bool bOnlyOnUndo = false;
-		RetriangulatePolygons( AffectedPolygons, bOnlyOnUndo );
-	}
-}
-
-
-void UEditableMesh::FindPolygonLoop( const FEdgeID EdgeID, TArray<FEdgeID>& OutEdgeLoopEdgeIDs, TArray<FEdgeID>& OutFlippedEdgeIDs, TArray<FEdgeID>& OutReversedEdgeIDPathToTake, TArray<FPolygonRef>& OutPolygonRefsToSplit ) const
-{
-	OutEdgeLoopEdgeIDs.Reset();
-	OutFlippedEdgeIDs.Reset();
-	OutReversedEdgeIDPathToTake.Reset();
-	OutPolygonRefsToSplit.Reset();
-
-	// Is the edge we're starting on a border edge?
-	bool bStartedOnBorderEdge = ( GetEdgeConnectedPolygonCount( EdgeID ) <= 1 );
-
-	// We'll actually do two passes searching for edges.  The first time, we'll flow along the polygons looking for a border edge.
-	// If we find one, it means we won't really have a loop, but instead we're just splitting a series of connected polygons.  In that
-	// case, the polygon with the border edge will become the start of our search for an opposing border edge.  Otherwise, we'll start
-	// at the input polygon and flow around until we come back to that polygon again.  If we can't make it back, then no polygons
-	// will be split by this operation.
-	bool bIsSearchingForBorderEdge = !bStartedOnBorderEdge;
-
-	// Keep track of whether we actually looped back around to the starting edge (rather than simply splitting across a series
-	// of polygons that both end at border edges.)
-	bool bIsCompleteLoop = false;
-
-	FEdgeID CurrentEdgeID = EdgeID;
-	bool bCurrentEdgeIsBorderEdge = bStartedOnBorderEdge;
-	bool bCurrentEdgeIsInOppositeDirection = false;
-	bool bCurrentEdgeIsInOppositeDirectionFromStartEdge = false;
-	for( ; ; )
-	{
-		// Add the current edge!
-		checkSlow( !OutEdgeLoopEdgeIDs.Contains( CurrentEdgeID ) );
-		OutEdgeLoopEdgeIDs.Add( CurrentEdgeID );
-		if( bCurrentEdgeIsInOppositeDirection )
-		{
-			OutFlippedEdgeIDs.Add( CurrentEdgeID );
-		}
-
-		FVertexID CurrentEdgeVertexIDs[ 2 ];
-		GetEdgeVertices( CurrentEdgeID, /* Out */ CurrentEdgeVertexIDs[ 0 ], /* Out */ CurrentEdgeVertexIDs[ 1 ] );
-
-		const FVector CurrentEdgeVertex0 = GetVertexAttribute( CurrentEdgeVertexIDs[ 0 ], UEditableMeshAttribute::VertexPosition(), 0 );
-		const FVector CurrentEdgeVertex1 = GetVertexAttribute( CurrentEdgeVertexIDs[ 1 ], UEditableMeshAttribute::VertexPosition(), 0 );
-		const FVector CurrentEdgeDirection = ( CurrentEdgeVertex1 - CurrentEdgeVertex0 ).GetSafeNormal();
-
-		const FEdgeID NextEdgeIDInPath = OutReversedEdgeIDPathToTake.Num() > 0 ? OutReversedEdgeIDPathToTake.Pop( false ) : FEdgeID::Invalid;
-
-		FEdgeID BestEdgeID = FEdgeID::Invalid;
-		FPolygonRef BestEdgeSplitsPolygon = FPolygonRef::Invalid;
-		bool bBestEdgeIsInOppositeDirection = false;
-		bool bBestEdgeIsBorderEdge = false;
-		float LargestAbsDotProduct = -1.0f;
-
-		// Let's take a look at all of the polygons connected to this edge.  These will start our loop.
-		const int32 ConnectedPolygonCount = GetEdgeConnectedPolygonCount( CurrentEdgeID );
-		for( int32 ConnectedPolygonNumber = 0; ConnectedPolygonNumber < ConnectedPolygonCount; ++ConnectedPolygonNumber )
-		{
-			const FPolygonRef ConnectedPolygonRef = GetEdgeConnectedPolygon( CurrentEdgeID, ConnectedPolygonNumber );
-
-			// Don't bother looking at the last polygon that was added to our split list.  We never want to backtrack!
-			if( OutPolygonRefsToSplit.Num() == 0 || ConnectedPolygonRef != OutPolygonRefsToSplit.Last() )
-			{
-				static TArray<FEdgeID> CandidateEdgeIDs;
-				CandidateEdgeIDs.Reset();
-				GetPolygonPerimeterEdges( ConnectedPolygonRef, /* Out */ CandidateEdgeIDs );
-
-				// Which edge of the connected polygon will be at the other end of our split?
-				for( const FEdgeID CandidateEdgeID : CandidateEdgeIDs )
-				{
-					// Don't bother with the edge we just came from
-					if( CandidateEdgeID != CurrentEdgeID )
-					{
-						// If we need to follow a specific path, then do that
-						if( NextEdgeIDInPath == FEdgeID::Invalid || CandidateEdgeID == NextEdgeIDInPath )
-						{
-							FVertexID CandidateEdgeVertexIDs[ 2 ];
-							GetEdgeVertices( CandidateEdgeID, /* Out */ CandidateEdgeVertexIDs[ 0 ], /* Out */ CandidateEdgeVertexIDs[ 1 ] );
-
-							const bool bIsBorderEdge = ( GetEdgeConnectedPolygonCount( CandidateEdgeID ) == 1 );
-
-							const FVector CandidateEdgeVertex0 = GetVertexAttribute( CandidateEdgeVertexIDs[ 0 ], UEditableMeshAttribute::VertexPosition(), 0 );
-							const FVector CandidateEdgeVertex1 = GetVertexAttribute( CandidateEdgeVertexIDs[ 1 ], UEditableMeshAttribute::VertexPosition(), 0 );
-							const FVector CandidateEdgeDirection = ( CandidateEdgeVertex1 - CandidateEdgeVertex0 ).GetSafeNormal();
-
-							const float DotProduct = FVector::DotProduct( CurrentEdgeDirection, CandidateEdgeDirection );
-							const float AbsDotProduct = FMath::Abs( DotProduct );
-
-							const float SameEdgeDirectionDotEpsilon = 0.05f;	// @todo mesheditor tweak
-							if( FMath::IsNearlyEqual( AbsDotProduct, LargestAbsDotProduct, SameEdgeDirectionDotEpsilon ) )
-							{
-								// If the candidate edge directions are pretty much the same, we'll choose the edge that flows closest to the
-								// direction that we split the last polygon in
-								if( OutEdgeLoopEdgeIDs.Num() > 1 )
-								{
-									FVertexID LastSplitEdgeVertexIDs[ 2 ];
-									GetEdgeVertices( OutEdgeLoopEdgeIDs[ OutEdgeLoopEdgeIDs.Num() - 2 ], /* Out */ LastSplitEdgeVertexIDs[ 0 ], /* Out */ LastSplitEdgeVertexIDs[ 1 ] );
-
-									const FVector LastSplitEdgeVertex0 = GetVertexAttribute( LastSplitEdgeVertexIDs[ 0 ], UEditableMeshAttribute::VertexPosition(), 0 );
-									const FVector LastSplitEdgeVertex1 = GetVertexAttribute( LastSplitEdgeVertexIDs[ 1 ], UEditableMeshAttribute::VertexPosition(), 0 );
-
-									const FVector DirectionTowardCenterOfCurrentEdge =
-										( FMath::Lerp( CurrentEdgeVertex0, CurrentEdgeVertex1, 0.5f ) -
-											FMath::Lerp( LastSplitEdgeVertex0, LastSplitEdgeVertex1, 0.5f ) ).GetSafeNormal();
-
-									const FVector DirectionTowardCenterOfCandidateEdge =
-										( FMath::Lerp( CandidateEdgeVertex0, CandidateEdgeVertex1, 0.5f ) -
-											FMath::Lerp( CurrentEdgeVertex0, CurrentEdgeVertex1, 0.5f ) ).GetSafeNormal();
-									const float CandidateEdgeDot = FVector::DotProduct( DirectionTowardCenterOfCurrentEdge, DirectionTowardCenterOfCandidateEdge );
-
-									check( BestEdgeID != FEdgeID::Invalid );
-
-									FVertexID BestEdgeVertexIDs[ 2 ];
-									GetEdgeVertices( BestEdgeID, /* Out */ BestEdgeVertexIDs[ 0 ], /* Out */ BestEdgeVertexIDs[ 1 ] );
-
-									const FVector BestEdgeVertex0 = GetVertexAttribute( BestEdgeVertexIDs[ 0 ], UEditableMeshAttribute::VertexPosition(), 0 );
-									const FVector BestEdgeVertex1 = GetVertexAttribute( BestEdgeVertexIDs[ 1 ], UEditableMeshAttribute::VertexPosition(), 0 );
-
-									const FVector DirectionTowardCenterOfBestEdge =
-										( FMath::Lerp( BestEdgeVertex0, BestEdgeVertex1, 0.5f ) -
-											FMath::Lerp( CurrentEdgeVertex0, CurrentEdgeVertex1, 0.5f ) ).GetSafeNormal();
-
-									const float BestEdgeDot = FVector::DotProduct( DirectionTowardCenterOfCurrentEdge, DirectionTowardCenterOfBestEdge );
-
-									if( CandidateEdgeDot > BestEdgeDot )
-									{
-										BestEdgeID = CandidateEdgeID;
-										BestEdgeSplitsPolygon = ConnectedPolygonRef;
-										bBestEdgeIsInOppositeDirection = ( DotProduct < 0.0f );
-										bBestEdgeIsBorderEdge = bIsBorderEdge;
-										LargestAbsDotProduct = AbsDotProduct;
-									}
-								}
-								else
-								{
-									// Edge directions are the same, but this is the very first split so we don't have a "flow" direction yet.
-									// Go ahead and prefer the edge that is closer to the initial edge.  This helps in the (uncommon) case of 
-									// multiple colinear edges on the same polygon (such as after SplitEdge() is called to insert a vertex on a polygon.)
-									const float BestEdgeDistance = [this, BestEdgeID, CurrentEdgeVertex0, CurrentEdgeVertex1]() -> float
-									{
-										check( BestEdgeID != FEdgeID::Invalid );
-
-										FVertexID BestEdgeVertexIDs[ 2 ];
-										GetEdgeVertices( BestEdgeID, /* Out */ BestEdgeVertexIDs[ 0 ], /* Out */ BestEdgeVertexIDs[ 1 ] );
-
-										const FVector BestEdgeVertex0 = GetVertexAttribute( BestEdgeVertexIDs[ 0 ], UEditableMeshAttribute::VertexPosition(), 0 );
-										const FVector BestEdgeVertex1 = GetVertexAttribute( BestEdgeVertexIDs[ 1 ], UEditableMeshAttribute::VertexPosition(), 0 );
-
-										FVector ClosestPoint0, ClosestPoint1;
-										FMath::SegmentDistToSegmentSafe( CurrentEdgeVertex0, CurrentEdgeVertex1, BestEdgeVertex0, BestEdgeVertex1, /* Out */ ClosestPoint0, /* Out */ ClosestPoint1 );
-										return ( ClosestPoint1 - ClosestPoint0 ).Size();
-									}( );
-
-									const float CandidateEdgeDistance = [CurrentEdgeVertex0, CurrentEdgeVertex1, CandidateEdgeVertex0, CandidateEdgeVertex1]() -> float
-									{
-										FVector ClosestPoint0, ClosestPoint1;
-										FMath::SegmentDistToSegmentSafe( CurrentEdgeVertex0, CurrentEdgeVertex1, CandidateEdgeVertex0, CandidateEdgeVertex1, /* Out */ ClosestPoint0, /* Out */ ClosestPoint1 );
-										return ( ClosestPoint1 - ClosestPoint0 ).Size();
-									}( );
-
-									if( CandidateEdgeDistance < BestEdgeDistance )
-									{
-										BestEdgeID = CandidateEdgeID;
-										BestEdgeSplitsPolygon = ConnectedPolygonRef;
-										bBestEdgeIsInOppositeDirection = ( DotProduct < 0.0f );
-										bBestEdgeIsBorderEdge = bIsBorderEdge;
-										LargestAbsDotProduct = AbsDotProduct;
-									}
-								}
-							}
-							else if( AbsDotProduct > LargestAbsDotProduct )
-							{
-								// This edge angle is the closest to our current edge so far!
-								BestEdgeID = CandidateEdgeID;
-								BestEdgeSplitsPolygon = ConnectedPolygonRef;
-								bBestEdgeIsInOppositeDirection = ( DotProduct < 0.0f );
-								bBestEdgeIsBorderEdge = bIsBorderEdge;
-								LargestAbsDotProduct = AbsDotProduct;
-							}
-						}
-					}
-				}
-			}
-		}
-
-		if( BestEdgeID != FEdgeID::Invalid &&
-			!OutPolygonRefsToSplit.Contains( BestEdgeSplitsPolygon ) )	// If we try to re-split the same polygon twice, then this loop is not valid
-		{
-			// OK, this polygon will definitely be split
-			OutPolygonRefsToSplit.Add( BestEdgeSplitsPolygon );
-
-			CurrentEdgeID = BestEdgeID;
-			bCurrentEdgeIsBorderEdge = bBestEdgeIsBorderEdge;
-			bCurrentEdgeIsInOppositeDirection = bBestEdgeIsInOppositeDirection;
-			if( bBestEdgeIsInOppositeDirection )
-			{
-				bCurrentEdgeIsInOppositeDirectionFromStartEdge = !bCurrentEdgeIsInOppositeDirectionFromStartEdge;
-			}
-
-			// Is the best edge already part of our loop?  If so, then we're done!
-			if( OutEdgeLoopEdgeIDs[ 0 ] == BestEdgeID )
-			{
-				bIsCompleteLoop = true;
-				break;
-			}
-			else if( OutEdgeLoopEdgeIDs.Contains( BestEdgeID ) )	// @todo mesheditor perf: Might want to use a TSet here, but we'll still need to keep a list too (order is important.)
-			{
-				// We ended up back at an edge that we already split, but it wasn't the edge that we started on.  This is
-				// not a valid loop, so clear our path and bail out.
-
-				// @todo mesheditor edgeloop: We need to revisit how we handle border edges and non-manifold edges.  Basically instead of searching
-				// for border edges, we should try looping around the "back" of polygons we've already visited, trying to get back to the starting
-				// edge.  This will allow the case of a cube mesh with a single non-manifold face to be split properly.  We just need to make sure
-				// we're not trying to split the same edge multiple times in different locations.  Probably, we don't want to split the same polygon
-				// more than once either.
-				OutEdgeLoopEdgeIDs.Reset();
-				OutFlippedEdgeIDs.Reset();
-				OutPolygonRefsToSplit.Reset();
-
-				break;
-			}
-			else if( bBestEdgeIsBorderEdge && bIsSearchingForBorderEdge )
-			{
-				// We found a border edge, so stop the search.  We'll now start over at this edge to form our loop.
-				bStartedOnBorderEdge = true;
-
-				bIsSearchingForBorderEdge = false;
-				bCurrentEdgeIsInOppositeDirection = bCurrentEdgeIsInOppositeDirectionFromStartEdge;
-
-				// Follow the path we took to get here, in reverse order, to make sure we get back to the edge we 
-				// were asked to create a loop on
-				OutReversedEdgeIDPathToTake = OutEdgeLoopEdgeIDs;
-
-				bIsCompleteLoop = false;
-
-				OutEdgeLoopEdgeIDs.Reset();
-				OutFlippedEdgeIDs.Reset();
-				OutPolygonRefsToSplit.Reset();
-			}
-			else
-			{
-				// Proceed to the next edge and try to continue the loop.  If we're at a border edge, the loop will definitely end here.
-			}
-		}
-		else
-		{
-			if( bStartedOnBorderEdge && bCurrentEdgeIsBorderEdge )
-			{
-				// We started on a border edge, and we've found the border edge on the other side of the polygons we'll be splitting.
-				// This isn't actually a loop, but we'll still split the polygons.
-			}
-			else
-			{
-				// We couldn't even find another edge, so we're done.
-				OutEdgeLoopEdgeIDs.Reset();
-				OutFlippedEdgeIDs.Reset();
-				OutPolygonRefsToSplit.Reset();
-			}
-
-			break;
-		}
-	}
-
-	// We're always splitting the same number of polygons as we have edges in the loop (these can be zero), except
-	// in the border edge case, where we're always splitting one less polygon.
-	if( bStartedOnBorderEdge &&
-		!bIsCompleteLoop )	// It's possible to start and finish on the same border edge, which is a complete loop
-	{
-		// We're splitting a series of polygons between two border edges
-		check( ( OutEdgeLoopEdgeIDs.Num() == 0 && OutPolygonRefsToSplit.Num() == 0 ) || ( OutEdgeLoopEdgeIDs.Num() == ( OutPolygonRefsToSplit.Num() + 1 ) ) );
-	}
-	else
-	{
-		// We're splitting polygons spanning a full loop of edges.  The starting edge is the same as the beginning edge.
-		check( OutEdgeLoopEdgeIDs.Num() == OutPolygonRefsToSplit.Num() );
-	}
-}
-
-
-void UEditableMesh::InsertEdgeLoop( const FEdgeID EdgeID, const TArray<float>& Splits, TArray<FEdgeID>& OutNewEdgeIDs )
-{
-	OutNewEdgeIDs.Reset();
-
-	// NOTE: The incoming splits should always be between 0.0 and 1.0, representing progress along 
-	//       the edge from the edge's first vertex toward it's other vertex.  The order doesn't matter (we'll sort them.)
-
-	// @todo mesheditor: Test with concave polygons -- probably we need to disallow/avoid splits that cross a concave part
-
-
-	static TArray<FEdgeID> EdgeLoopEdgeIDs;
-	EdgeLoopEdgeIDs.Reset();
-
-	static TArray<FEdgeID> FlippedEdgeIDs;
-	FlippedEdgeIDs.Reset();
-
-	static TArray<FEdgeID> ReversedEdgeIDPathToTake;
-	ReversedEdgeIDPathToTake.Reset();
-
-	static TArray<FPolygonRef> PolygonRefsToSplit;
-	PolygonRefsToSplit.Reset();
-
-
-	FindPolygonLoop( 
-		EdgeID, 
-		/* Out */ EdgeLoopEdgeIDs,
-		/* Out */ FlippedEdgeIDs,
-		/* Out */ ReversedEdgeIDPathToTake,
-		/* Out */ PolygonRefsToSplit );
-
-	static TSet<FEdgeID> FlippedEdgeIDSet;
-	FlippedEdgeIDSet.Reset();
-	FlippedEdgeIDSet.Append( FlippedEdgeIDs );
-
-	check( Splits.Num() > 0 );
-
-	// Sort the split values smallest to largest.  We'll be adding a strip of vertices for each split, and the
-	// IDs for those new vertices need to be in order.
-	static TArray<float> SortedSplits;
-	SortedSplits.Reset();
-	SortedSplits = Splits;
-	if( SortedSplits.Num() > 1 )
-	{
-		SortedSplits.Sort();
-	}
-
-
-	if( PolygonRefsToSplit.Num() > 0 )
-	{
-		// Keep track of the new vertices create by splitting all of the edges.  For each edge we split, an array of
-		// vertex IDs for each split along that edge
-		static TArray<TArray<FVertexID>> NewVertexIDsForEachEdge;
-		NewVertexIDsForEachEdge.Reset();
-
-
-		// Now let's go through and create new vertices for the loops by splitting edges
-		{
-			for( int32 EdgeLoopEdgeNumber = 0; EdgeLoopEdgeNumber < EdgeLoopEdgeIDs.Num(); ++EdgeLoopEdgeNumber )
-			{
-				const FEdgeID EdgeLoopEdgeID = EdgeLoopEdgeIDs[ EdgeLoopEdgeNumber ];
-
-				// If the edge winds in the opposite direction from the last edge, we'll need to flip the split positions around
-				const bool bIsFlipped = FlippedEdgeIDSet.Contains( EdgeLoopEdgeID );
-				if( bIsFlipped )
-				{
-					static TArray<float> TempSplits;
-					TempSplits.SetNumUninitialized( SortedSplits.Num(), false );
-					for( int32 SplitIndex = 0; SplitIndex < SortedSplits.Num(); ++SplitIndex )
-					{
-						TempSplits[ SplitIndex ] = 1.0f - SortedSplits[ ( SortedSplits.Num() - 1 ) - SplitIndex ];
-					}
-					SortedSplits = TempSplits;
-				}
-
-				// Split this edge
-				// @todo mesheditor urgent edgeloop: We need to make sure we're not retriangulating and calculating normals over and over again for the same faces.
-				static TArray<FVertexID> CurrentVertexIDs;
-				CurrentVertexIDs.Reset();
-				SplitEdge( EdgeLoopEdgeID, SortedSplits, /* Out */ CurrentVertexIDs );
-
-				// If the edge winding is backwards, we'll reverse the order of the vertex IDs in our list
-				// @todo mesheditor urgent edgeloop: INCORRECT (bIsFlipped is supposed to swap from last time -- TempSplits.)  Also not sure if needed or not... might be double compensating here.
-				static TArray<FVertexID> NewVertexIDsForEdge;
-				NewVertexIDsForEdge.SetNum( CurrentVertexIDs.Num(), false );
-				for( int32 VertexNumber = 0; VertexNumber < CurrentVertexIDs.Num(); ++VertexNumber )
-				{
-					NewVertexIDsForEdge[ /*bIsFlipped ? */( ( CurrentVertexIDs.Num() - VertexNumber ) - 1 ) /*: VertexNumber */] = CurrentVertexIDs[ VertexNumber ];
-				}
-				NewVertexIDsForEachEdge.Add( NewVertexIDsForEdge );
-			}
-		}
-
-
-		// Time to create new polygons for the split faces (and delete the old ones)
-		{
-			static TArray<FPolygonToSplit> PolygonsToSplit;
-			PolygonsToSplit.Reset();
-
-			for( int32 PolygonToSplitIter = 0; PolygonToSplitIter < PolygonRefsToSplit.Num(); ++PolygonToSplitIter )
-			{
-				const FPolygonRef PolygonRef = PolygonRefsToSplit[ PolygonToSplitIter ];
-
-				FPolygonToSplit& PolygonToSplit = *new( PolygonsToSplit ) FPolygonToSplit();
-				PolygonToSplit.PolygonRef = PolygonRef;
-
-				// The first and second edges connected to this polygon that are being split up
-				const int32 FirstEdgeNumber = PolygonToSplitIter;
-				const int32 SecondEdgeNumber = ( PolygonToSplitIter + 1 ) % EdgeLoopEdgeIDs.Num();
-
-				const FEdgeID FirstSplitEdgeID = EdgeLoopEdgeIDs[ FirstEdgeNumber ];
-				const FEdgeID SecondSplitEdgeID = EdgeLoopEdgeIDs[ SecondEdgeNumber ];
-				check( FirstSplitEdgeID != SecondSplitEdgeID );
-
-				// The (ordered) list of new vertices that was created by splitting the first and second edge.  One for each split.
-				const TArray<FVertexID>& FirstSplitEdgeNewVertexIDs = NewVertexIDsForEachEdge[ FirstEdgeNumber ];
-				const TArray<FVertexID>& SecondSplitEdgeNewVertexIDs = NewVertexIDsForEachEdge[ SecondEdgeNumber ];
-
-				for( int32 SplitIter = 0; SplitIter < SortedSplits.Num(); ++SplitIter )
-				{
-					FVertexPair& NewVertexPair = *new( PolygonToSplit.VertexPairsToSplitAt ) FVertexPair();
-					NewVertexPair.VertexID0 = FirstSplitEdgeNewVertexIDs[ SplitIter ];
-					NewVertexPair.VertexID1 = SecondSplitEdgeNewVertexIDs[ SplitIter ];
-				}
-			}
-
-
-			// Actually split up the polygons
-			static TArray<FEdgeID> NewEdgeIDs;
-			NewEdgeIDs.Reset();
-			SplitPolygons( PolygonsToSplit, /* Out */ NewEdgeIDs );
-
-			OutNewEdgeIDs.Append( NewEdgeIDs );
-		}
-	}
-}
-
-
 void UEditableMesh::SplitPolygons( const TArray<FPolygonToSplit>& PolygonsToSplit, TArray<FEdgeID>& OutNewEdgeIDs )
 {
 	OutNewEdgeIDs.Reset();
@@ -2699,12 +3665,12 @@
 
 	for( const FPolygonToSplit& PolygonToSplit : PolygonsToSplit )
 	{
-		const FPolygonRef PolygonRef = PolygonToSplit.PolygonRef;
+		const FPolygonID PolygonID = PolygonToSplit.PolygonID;
 
 		// Get all of the polygon's vertices
 		static TArray<FVertexID> PerimeterVertexIDs;
 		PerimeterVertexIDs.Reset();
-		GetPolygonPerimeterVertices( PolygonRef, /* Out */ PerimeterVertexIDs );
+		GetPolygonPerimeterVertices( PolygonID, /* Out */ PerimeterVertexIDs );
 		
 
 		// Figure out where exactly we're splitting the polygon for these splits.  Remember, we support splitting the
@@ -2731,7 +3697,8 @@
 
 
 			FPolygonToCreate& NewPolygon = *new( PolygonsToCreate ) FPolygonToCreate();
-			NewPolygon.SectionID = PolygonRef.SectionID;
+			NewPolygon.PolygonGroupID = GetGroupForPolygon( PolygonID );
+			NewPolygon.PolygonEdgeHardness = EPolygonEdgeHardness::NewEdgesSoft;
 
 			static TArray<int32> PerimeterVertexNumbers;
 			PerimeterVertexNumbers.Reset();
@@ -2793,23 +3760,8 @@
 			for( const int32 VertexNumber : PerimeterVertexNumbers )
 			{
 				FVertexAndAttributes& NewVertexAndAttributes = *new( NewPolygon.PerimeterVertices ) FVertexAndAttributes();
-				NewVertexAndAttributes.VertexID = PerimeterVertexIDs[ VertexNumber ];
-
-				// Copy all of the polygon vertex attributes over
-				{
-					for( int32 TextureCoordinateIndex = 0; TextureCoordinateIndex < TextureCoordinateCount; ++TextureCoordinateIndex )
-					{
-						NewVertexAndAttributes.PolygonVertexAttributes.Attributes.Add( FMeshElementAttributeData(
-							UEditableMeshAttribute::VertexTextureCoordinate(),
-							TextureCoordinateIndex,
-							GetPolygonPerimeterVertexAttribute( PolygonRef, VertexNumber, UEditableMeshAttribute::VertexTextureCoordinate(), TextureCoordinateIndex ) ) );
-					}
-
-					NewVertexAndAttributes.PolygonVertexAttributes.Attributes.Add( FMeshElementAttributeData(
-						UEditableMeshAttribute::VertexColor(),
-						0,
-						GetPolygonPerimeterVertexAttribute( PolygonRef, VertexNumber, UEditableMeshAttribute::VertexColor(), 0 ) ) );
-				}
+				NewVertexAndAttributes.VertexID = FVertexID::Invalid;
+				NewVertexAndAttributes.VertexInstanceID = GetPolygonPerimeterVertexInstance( PolygonID, VertexNumber );
 			}
 
 			LastPolygonVertexNumbers[ 0 ] = PerimeterVertexNumbers[ 0 ];
@@ -2818,53 +3770,50 @@
 		}
 	}
 
-	// Create new polygons that are split appropriately and connect to the new vertices we've added
-	static TArray< FPolygonRef > NewPolygonRefs;
-	NewPolygonRefs.Reset();
-	static TArray< FEdgeID > NewEdgeIDs;
-	NewEdgeIDs.Reset();
-	CreatePolygons( PolygonsToCreate, /* Out */ NewPolygonRefs, /* Out */ NewEdgeIDs );
-
-	OutNewEdgeIDs.Append( NewEdgeIDs );
-
-
 	// Delete the old polygons
 	{
-		static TArray<FPolygonRef> PolygonRefsToDelete;
-		PolygonRefsToDelete.Reset();
-		PolygonRefsToDelete.Reserve( PolygonsToSplit.Num() );
+		static TArray<FPolygonID> PolygonIDsToDelete;
+		PolygonIDsToDelete.Reset();
+		PolygonIDsToDelete.Reserve( PolygonsToSplit.Num() );
 		for( const FPolygonToSplit& PolygonToSplit : PolygonsToSplit )
 		{
-			PolygonRefsToDelete.Add( PolygonToSplit.PolygonRef );
+			PolygonIDsToDelete.Add( PolygonToSplit.PolygonID );
 		}
 
 		const bool bDeleteOrphanEdges = false;
 		const bool bDeleteOrphanVertices = false;
-		const bool bDeleteEmptySections = false;
-		DeletePolygons( PolygonRefsToDelete, bDeleteOrphanEdges, bDeleteOrphanVertices, bDeleteEmptySections );
+		const bool bDeleteOrphanVertexInstances = false;
+		const bool bDeleteEmptyPolygonGroups = false;
+		DeletePolygons( PolygonIDsToDelete, bDeleteOrphanEdges, bDeleteOrphanVertices, bDeleteOrphanVertexInstances, bDeleteEmptyPolygonGroups );
 	}
 
-	// Generate normals and tangents
-	GenerateNormalsAndTangentsForPolygons( NewPolygonRefs );
+	// Create new polygons that are split appropriately and connect to the new vertices we've added
+	static TArray< FPolygonID > NewPolygonIDs;
+	NewPolygonIDs.Reset();
+	static TArray< FEdgeID > NewEdgeIDs;
+	NewEdgeIDs.Reset();
+	CreatePolygons( PolygonsToCreate, /* Out */ NewPolygonIDs, /* Out */ NewEdgeIDs );
+
+	OutNewEdgeIDs.Append( NewEdgeIDs );
 }
 
 
-void UEditableMesh::DeleteEdgeAndConnectedPolygons( const FEdgeID EdgeID, const bool bDeleteOrphanedEdges, const bool bDeleteOrphanedVertices, const bool bDeleteEmptySections )
+void UEditableMesh::DeleteEdgeAndConnectedPolygons( const FEdgeID EdgeID, const bool bDeleteOrphanedEdges, const bool bDeleteOrphanedVertices, const bool bDeleteOrphanedVertexInstances, const bool bDeleteEmptyPolygonGroups )
 {
-	static TArray<FPolygonRef> PolygonRefsToDelete;
-	PolygonRefsToDelete.Reset();
+	static TArray<FPolygonID> PolygonIDsToDelete;
+	PolygonIDsToDelete.Reset();
 
 	const int32 ConnectedPolygonCount = GetEdgeConnectedPolygonCount( EdgeID );
 	for( int32 ConnectedPolygonNumber = 0; ConnectedPolygonNumber < ConnectedPolygonCount; ++ConnectedPolygonNumber )
 	{
-		const FPolygonRef PolygonRef = GetEdgeConnectedPolygon( EdgeID, ConnectedPolygonNumber );
+		const FPolygonID PolygonID = GetEdgeConnectedPolygon( EdgeID, ConnectedPolygonNumber );
 
 		// Although it can be uncommon, it's possible the edge is connecting the same polygon to itself.  We need to add uniquely.
-		PolygonRefsToDelete.AddUnique( PolygonRef );
+		PolygonIDsToDelete.AddUnique( PolygonID );
 	}
 
 	// Delete the polygons
-	DeletePolygons( PolygonRefsToDelete, bDeleteOrphanedEdges, bDeleteOrphanedVertices, bDeleteEmptySections );
+	DeletePolygons( PolygonIDsToDelete, bDeleteOrphanedEdges, bDeleteOrphanedVertices, bDeleteOrphanedVertexInstances, bDeleteEmptyPolygonGroups );
 
 	// If the caller asked us not deleted orphaned edges, our edge-to-delete will still be hanging around.  Let's go
 	// and delete it now.
@@ -2883,7 +3832,7 @@
 }
 
 
-void UEditableMesh::DeleteVertexAndConnectedEdgesAndPolygons( const FVertexID VertexID, const bool bDeleteOrphanedEdges, const bool bDeleteOrphanedVertices, const bool bDeleteEmptySections )
+void UEditableMesh::DeleteVertexAndConnectedEdgesAndPolygons( const FVertexID VertexID, const bool bDeleteOrphanedEdges, const bool bDeleteOrphanedVertices, const bool bDeleteOrphanedVertexInstances, const bool bDeleteEmptyPolygonGroups )
 {
 	static TArray<FEdgeID> EdgeIDsToDelete;
 	EdgeIDsToDelete.Reset();
@@ -2901,35 +3850,21 @@
 		// a previous iteration through this loop.
 		if( IsValidEdge( EdgeIDToDelete ) )
 		{
-			DeleteEdgeAndConnectedPolygons( EdgeIDToDelete, bDeleteOrphanedEdges, bDeleteOrphanedVertices, bDeleteEmptySections );
+			DeleteEdgeAndConnectedPolygons( EdgeIDToDelete, bDeleteOrphanedEdges, bDeleteOrphanedVertices, bDeleteOrphanedVertexInstances, bDeleteEmptyPolygonGroups );
 		}
 	}
 }
 
 
-void UEditableMesh::CreateEmptyVertexRange( const int32 NumVerticesToAdd, TArray<FVertexID>& OutNewVertexIDs )
-{
-	CreateEmptyVertexRange_Internal( NumVerticesToAdd, nullptr, OutNewVertexIDs );
-
-	// NOTE: We iterate backwards, to delete vertices in the opposite order that we added them
-	FDeleteOrphanVerticesChangeInput RevertInput;
-	for( int32 VertexIter = NumVerticesToAdd - 1; VertexIter >= 0; --VertexIter )
-	{
-		const FVertexID VertexIDToDelete = OutNewVertexIDs[ VertexIter ];
-		RevertInput.VertexIDsToDelete.Add( VertexIDToDelete );
-	}
-
-	AddUndo( MakeUnique<FDeleteOrphanVerticesChange>( MoveTemp( RevertInput ) ) );
-}
-
-
 void UEditableMesh::DeleteOrphanVertices( const TArray<FVertexID>& VertexIDsToDelete )
 {
-	FCreateVerticesChangeInput RevertInput;
-
+	UE_LOG( LogEditableMesh, Verbose, TEXT( "DeleteOrphanVertices: %s" ), *LogHelpers::ArrayToString( VertexIDsToDelete ) );
 
 	// Back everything up
 	{
+		FCreateVerticesChangeInput RevertInput;
+		RevertInput.VerticesToCreate.Reserve( VertexIDsToDelete.Num() );
+
 		// NOTE: We iterate backwards, to restore vertices in the opposite order that we deleted them
 		for( int32 VertexNumber = VertexIDsToDelete.Num() - 1; VertexNumber >= 0; --VertexNumber )
 		{
@@ -2937,9 +3872,13 @@
 
 			// Make sure the vertex is truly an orphan.  We're not going to be able to restore it's polygon vertex attributes,
 			// because the polygons won't exist when we're restoring the change
-			check( GetVertexConnectedEdgeCount( VertexID ) == 0 );
+			const FMeshVertex& Vertex = Vertices[ VertexID.GetValue() ];
+			check( Vertex.ConnectedEdgeIDs.Num() == 0 );
+			check( Vertex.VertexInstanceIDs.Num() == 0 );
 
-			FVertexToCreate& VertexToCreate = *new( RevertInput.VerticesToCreate ) FVertexToCreate();
+			RevertInput.VerticesToCreate.Emplace();
+			FVertexToCreate& VertexToCreate = RevertInput.VerticesToCreate.Last();
+
 			VertexToCreate.OriginalVertexID = VertexID;
 
 			for( const FName AttributeName : UEditableMesh::GetValidVertexAttributes() )
@@ -2947,317 +3886,953 @@
 				const int32 MaxAttributeIndex = GetMaxAttributeIndex( AttributeName );
 				for( int32 AttributeIndex = 0; AttributeIndex < MaxAttributeIndex; ++AttributeIndex )
 				{
-					FMeshElementAttributeData& VertexAttribute = *new( VertexToCreate.VertexAttributes.Attributes ) FMeshElementAttributeData(
+					VertexToCreate.VertexAttributes.Attributes.Emplace(
 						AttributeName,
 						AttributeIndex,
-						GetVertexAttribute( VertexID, AttributeName, AttributeIndex ) );
+						GetVertexAttribute( VertexID, AttributeName, AttributeIndex )
+					);
 				}
 			}
 		}
+
+		AddUndo( MakeUnique<FCreateVerticesChange>( MoveTemp( RevertInput ) ) );
 	}
 
-	DeleteOrphanVertices_Internal( VertexIDsToDelete );
+	// Give the adapter a chance to handle this before they are deleted
+	for( UEditableMeshAdapter* Adapter : Adapters )
+	{
+		Adapter->OnDeleteOrphanVertices( this, VertexIDsToDelete );
+	}
+	
+	// Actually delete the vertices
+	{
+		for( const FVertexID VertexIDToDelete : VertexIDsToDelete )
+		{
+			Vertices.RemoveAt( VertexIDToDelete.GetValue() );
+		}
+	}
 
-	AddUndo( MakeUnique<FCreateVerticesChange>( MoveTemp( RevertInput ) ) );
+	UE_LOG( LogEditableMesh, Verbose, TEXT( "* DeleteOrphanVertices returned" ) );
 }
 
 
+void UEditableMesh::DeleteVertexInstances( const TArray<FVertexInstanceID>& VertexInstanceIDsToDelete, const bool bDeleteOrphanedVertices )
+{
+	UE_LOG( LogEditableMesh, Verbose, TEXT( "DeleteVertexInstances: %s" ), *LogHelpers::ArrayToString( VertexInstanceIDsToDelete ) );
+
+	// Back everything up
+	{
+		FCreateVertexInstancesChangeInput RevertInput;
+		RevertInput.VertexInstancesToCreate.Reserve( VertexInstanceIDsToDelete.Num() );
+
+		// NOTE: We iterate backwards, to restore vertices in the opposite order that we deleted them
+		for( int32 VertexNumber = VertexInstanceIDsToDelete.Num() - 1; VertexNumber >= 0; --VertexNumber )
+		{
+			const FVertexInstanceID VertexInstanceID = VertexInstanceIDsToDelete[ VertexNumber ];
+			const FMeshVertexInstance& VertexInstance = VertexInstances[ VertexInstanceID.GetValue() ];
+
+			// Back up properties
+			RevertInput.VertexInstancesToCreate.Emplace();
+			FVertexInstanceToCreate& VertexInstanceToCreate = RevertInput.VertexInstancesToCreate.Last();
+
+			VertexInstanceToCreate.VertexID = VertexInstance.VertexID;
+			VertexInstanceToCreate.OriginalVertexInstanceID = VertexInstanceID;
+
+			for( const FName AttributeName : GetValidVertexInstanceAttributes() )
+			{
+				const int32 MaxAttributeIndex = GetMaxAttributeIndex( AttributeName );
+				for( int32 AttributeIndex = 0; AttributeIndex < MaxAttributeIndex; ++AttributeIndex )
+				{
+					VertexInstanceToCreate.VertexInstanceAttributes.Attributes.Emplace(
+						AttributeName,
+						AttributeIndex,
+						GetVertexInstanceAttribute( VertexInstanceID, AttributeName, AttributeIndex )
+					);
+				}
+			}
+		}
+
+		AddUndo( MakeUnique<FCreateVertexInstancesChange>( MoveTemp( RevertInput ) ) );
+	}
+
+	// Give the adapter a chance to do something with this event before it happens
+	for( UEditableMeshAdapter* Adapter : Adapters )
+	{
+		Adapter->OnDeleteVertexInstances( this, VertexInstanceIDsToDelete );
+	}
+
+	static TArray<FVertexID> OrphanedVertexIDs;
+	OrphanedVertexIDs.Reset();
+
+	// Actually delete the vertex instances
+	{
+		for( const FVertexInstanceID VertexInstanceIDToDelete : VertexInstanceIDsToDelete )
+		{
+			// Check that no polygon is referencing this vertex instance
+			check( GetVertexInstanceConnectedPolygonCount( VertexInstanceIDToDelete ) == 0 );
+
+			// Remove the reference in the vertex it instances
+			const FMeshVertexInstance& VertexInstance = VertexInstances[ VertexInstanceIDToDelete.GetValue() ];
+			FMeshVertex& Vertex = Vertices[ VertexInstance.VertexID.GetValue() ];
+			verify( Vertex.VertexInstanceIDs.Remove( VertexInstanceIDToDelete ) == 1 );
+
+			if( bDeleteOrphanedVertices )
+			{
+				if( Vertex.VertexInstanceIDs.Num() == 0 )
+				{
+					OrphanedVertexIDs.Add( VertexInstance.VertexID );
+				}
+			}
+
+			// Delete the element from the container
+			VertexInstances.RemoveAt( VertexInstanceIDToDelete.GetValue() );
+		}
+	}
+
+	// Delete orphaned vertices, if there are any.
+	if( OrphanedVertexIDs.Num() > 0 )
+	{
+		DeleteOrphanVertices( OrphanedVertexIDs );
+	}
+
+	UE_LOG( LogEditableMesh, Verbose, TEXT( "* DeleteVertexInstances returned" ) );
+}
+
+
 void UEditableMesh::DeleteEdges( const TArray<FEdgeID>& EdgeIDsToDelete, const bool bDeleteOrphanedVertices )
 {
+	UE_LOG( LogEditableMesh, Verbose, TEXT( "DeleteEdges: %s" ), *LogHelpers::ArrayToString( EdgeIDsToDelete ) );
+
 	// Back everything up
-	FCreateEdgesChangeInput RevertInput;
 	{
+		FCreateEdgesChangeInput RevertInput;
+
 		// NOTE: We iterate backwards, to restore edges in the opposite order that we deleted them
 		for( int32 EdgeNumber = EdgeIDsToDelete.Num() - 1; EdgeNumber >= 0; --EdgeNumber )
 		{
 			const FEdgeID EdgeID = EdgeIDsToDelete[ EdgeNumber ];
+			const FMeshEdge& Edge = Edges[ EdgeID.GetValue() ];
 
-			FEdgeToCreate& EdgeToCreate = *new( RevertInput.EdgesToCreate ) FEdgeToCreate();
+			RevertInput.EdgesToCreate.Emplace();
+			FEdgeToCreate& EdgeToCreate = RevertInput.EdgesToCreate.Last();
+
 			EdgeToCreate.OriginalEdgeID = EdgeID;
-			EdgeToCreate.VertexID0 = GetEdgeVertex( EdgeID, 0 );
-			EdgeToCreate.VertexID1 = GetEdgeVertex( EdgeID, 1 );
-
-			const int32 ConnectedPolygonCount = GetEdgeConnectedPolygonCount( EdgeID );
-			EdgeToCreate.ConnectedPolygons.Reserve( ConnectedPolygonCount );
-			for( int32 ConnectedPolygonNumber = 0; ConnectedPolygonNumber < ConnectedPolygonCount; ++ConnectedPolygonNumber )
-			{
-				EdgeToCreate.ConnectedPolygons.Add( GetEdgeConnectedPolygon( EdgeID, ConnectedPolygonNumber ) );
-			}
+			EdgeToCreate.VertexID0 = Edge.VertexIDs[ 0 ];
+			EdgeToCreate.VertexID1 = Edge.VertexIDs[ 1 ];
+			EdgeToCreate.ConnectedPolygons = Edge.ConnectedPolygons;
 			
 			for( const FName AttributeName : UEditableMesh::GetValidEdgeAttributes() )
 			{
 				const int32 MaxAttributeIndex = GetMaxAttributeIndex( AttributeName );
 				for( int32 AttributeIndex = 0; AttributeIndex < MaxAttributeIndex; ++AttributeIndex )
 				{
-					FMeshElementAttributeData& EdgeAttribute = *new( EdgeToCreate.EdgeAttributes.Attributes ) FMeshElementAttributeData(
+					EdgeToCreate.EdgeAttributes.Attributes.Emplace(
 						AttributeName,
 						AttributeIndex,
-						GetEdgeAttribute( EdgeID, AttributeName, AttributeIndex ) );
+						GetEdgeAttribute( EdgeID, AttributeName, AttributeIndex )
+					);
 				}
 			}
 		}
+
+		AddUndo( MakeUnique<FCreateEdgesChange>( MoveTemp( RevertInput ) ) );
 	}
 
+	for( UEditableMeshAdapter* Adapter : Adapters )
+	{
+		Adapter->OnDeleteEdges( this, EdgeIDsToDelete );
+	}
 
-	AddUndo( MoveTemp( MakeUnique<FCreateEdgesChange>( RevertInput ) ) );
-
 	// Delete the edges
-	DeleteEdges_Internal( EdgeIDsToDelete, bDeleteOrphanedVertices );
+	{
+		static TArray<FVertexID> OrphanedVertexIDs;
+		OrphanedVertexIDs.Reset();
+
+		for( const FEdgeID EdgeID : EdgeIDsToDelete )
+		{
+			FMeshEdge& Edge = Edges[ EdgeID.GetValue() ];
+
+			for( const FVertexID EdgeVertexID : Edge.VertexIDs )
+			{
+				FMeshVertex& Vertex = Vertices[ EdgeVertexID.GetValue() ];
+				verify( Vertex.ConnectedEdgeIDs.RemoveSingle( EdgeID ) == 1 );
+
+				// If the vertex has no more edges connected, we'll keep track of that so we can delete the vertex later
+				if( bDeleteOrphanedVertices && Vertex.ConnectedEdgeIDs.Num() == 0 )
+				{
+					check( Vertex.VertexInstanceIDs.Num() == 0 );	// We must already have deleted any vertex instances
+					check( !OrphanedVertexIDs.Contains( EdgeVertexID ) );	// Orphaned vertex shouldn't have already been orphaned by an earlier deleted edge passed into this function
+					OrphanedVertexIDs.Add( EdgeVertexID );
+				}
+			}
+
+			// Delete the edge
+			Edges.RemoveAt( EdgeID.GetValue() );
+		}
+
+		// If we orphaned any vertices and we were asked to delete those, then we'll go ahead and do that now.
+		if( OrphanedVertexIDs.Num() > 0 )
+		{
+			DeleteOrphanVertices( OrphanedVertexIDs );
+		}
+	}
+
+	UE_LOG( LogEditableMesh, Verbose, TEXT( "* DeleteEdges returned" ) );
 }
 
 
+void UEditableMesh::CreateEmptyVertexRange( const int32 NumVerticesToCreate, TArray<FVertexID>& OutNewVertexIDs )
+{
+	UE_LOG( LogEditableMesh, Verbose, TEXT( "CreateEmptyVertexRange: %d" ), NumVerticesToCreate );
+
+	OutNewVertexIDs.Reset( NumVerticesToCreate );
+
+	// Reserve elements
+	Vertices.Reserve( Vertices.Num() + NumVerticesToCreate );
+
+	for( int32 Count = 0; Count < NumVerticesToCreate; ++Count )
+	{
+		// Allocate vertex
+		const FVertexID VertexID = FVertexID( Vertices.Add( FMeshVertex() ) );
+		FMeshVertex& Vertex = Vertices[ VertexID.GetValue() ];
+
+		// Initialize vertex
+		Vertex.VertexPosition = FVector::ZeroVector;
+		Vertex.CornerSharpness = 0.0f;
+
+		OutNewVertexIDs.Add( VertexID );
+	}
+
+	// NOTE: We iterate backwards, to delete vertices in the opposite order that we added them
+	{
+		FDeleteOrphanVerticesChangeInput RevertInput;
+		RevertInput.VertexIDsToDelete.Reserve( NumVerticesToCreate );
+		for( int32 VertexNumber = OutNewVertexIDs.Num() - 1; VertexNumber >= 0; --VertexNumber )
+		{
+			RevertInput.VertexIDsToDelete.Add( OutNewVertexIDs[ VertexNumber ] );
+		}
+
+		AddUndo( MakeUnique<FDeleteOrphanVerticesChange>( MoveTemp( RevertInput ) ) );
+	}
+
+	// Advise the adapter of new vertices
+	for( UEditableMeshAdapter* Adapter : Adapters )
+	{
+		Adapter->OnCreateEmptyVertexRange( this, OutNewVertexIDs );
+	}
+
+	UE_LOG( LogEditableMesh, Verbose, TEXT( "* CreateEmptyVertexRange returned %s" ), *LogHelpers::ArrayToString( OutNewVertexIDs ) );
+}
+
+
 void UEditableMesh::CreateVertices( const TArray<FVertexToCreate>& VerticesToCreate, TArray<FVertexID>& OutNewVertexIDs )
 {
-	OutNewVertexIDs.Reset();
+	UE_LOG( LogEditableMesh, Verbose, TEXT( "CreateVertices: %s" ), *LogHelpers::ArrayToString( VerticesToCreate ) );
 
+	OutNewVertexIDs.Reset( VerticesToCreate.Num() );
+
+	// Reserve elements
+	Vertices.Reserve( Vertices.Num() + VerticesToCreate.Num() );
+
 	for( const FVertexToCreate& VertexToCreate : VerticesToCreate )
 	{
-		static TArray<FVertexID> OverrideVertexIDs;
-		OverrideVertexIDs.Reset();
-
-		if( VertexToCreate.OriginalVertexID != FVertexID::Invalid )
+		// Allocate vertex
+		FVertexID VertexID = VertexToCreate.OriginalVertexID;
+		if( VertexID != FVertexID::Invalid )
 		{
-			// Restore the vertex's original vertex ID
-			OverrideVertexIDs.Add( VertexToCreate.OriginalVertexID );
+			Vertices.Insert( VertexID.GetValue(), FMeshVertex() );
+		}
+		else
+		{
+			VertexID = FVertexID( Vertices.Add( FMeshVertex() ) );
 		}
 
-		// Create the vertex and set it's attributes
-		static TArray<FVertexID> NewVertexIDs;
-		NewVertexIDs.Reset();
-		CreateEmptyVertexRange_Internal( 1, OverrideVertexIDs.Num() > 0 ? &OverrideVertexIDs : nullptr, /* Out */ NewVertexIDs );
+		FMeshVertex& Vertex = Vertices[ VertexID.GetValue() ];
 
-		const FVertexID NewVertexID = NewVertexIDs[ 0 ];
-		OutNewVertexIDs.Add( NewVertexID );
+		// Initialize vertex
+		Vertex.VertexPosition = FVector::ZeroVector;
+		Vertex.CornerSharpness = 0.0f;
 
-		for( const FMeshElementAttributeData& VertexAttribute : VertexToCreate.VertexAttributes.Attributes )
+		OutNewVertexIDs.Add( VertexID );
+	}
+
+	// NOTE: We iterate backwards, to delete vertices in the opposite order that we added them
+	{
+		FDeleteOrphanVerticesChangeInput RevertInput;
+		RevertInput.VertexIDsToDelete.Reserve( VerticesToCreate.Num() );
+		for( int32 VertexNumber = OutNewVertexIDs.Num() - 1; VertexNumber >= 0; --VertexNumber )
 		{
-			SetVertexAttribute_Internal( NewVertexID, VertexAttribute.AttributeName, VertexAttribute.AttributeIndex, VertexAttribute.AttributeValue );
+			RevertInput.VertexIDsToDelete.Add( OutNewVertexIDs[ VertexNumber ] );
+		}
+
+		AddUndo( MakeUnique<FDeleteOrphanVerticesChange>( MoveTemp( RevertInput ) ) );
+	}
+
+	// Advise the adapter of new vertices
+	for( UEditableMeshAdapter* Adapter : Adapters )
+	{
+		Adapter->OnCreateVertices( this, OutNewVertexIDs );
+	}
+
+	// Set new vertex attributes
+	for( int32 Index = 0; Index < OutNewVertexIDs.Num(); ++Index )
+	{
+		for( const FMeshElementAttributeData& VertexAttribute : VerticesToCreate[ Index ].VertexAttributes.Attributes )
+		{
+			SetVertexAttribute( OutNewVertexIDs[ Index ], VertexAttribute.AttributeName, VertexAttribute.AttributeIndex, VertexAttribute.AttributeValue );
 		}
 	}
 
-	// NOTE: We iterate backwards, to delete vertices in the opposite order that we added them
-	FDeleteOrphanVerticesChangeInput RevertInput;
-	RevertInput.VertexIDsToDelete.Reserve( VerticesToCreate.Num() );
-	for( int32 VertexNumber = OutNewVertexIDs.Num() - 1; VertexNumber >= 0; --VertexNumber )
-	{
-		RevertInput.VertexIDsToDelete.Add( OutNewVertexIDs[ VertexNumber ] );
-	}
-
-	AddUndo( MakeUnique<FDeleteOrphanVerticesChange>( MoveTemp( RevertInput ) ) );
+	UE_LOG( LogEditableMesh, Verbose, TEXT( "* CreateVertices returned %s" ), *LogHelpers::ArrayToString( OutNewVertexIDs ) );
 }
 
 
+void UEditableMesh::CreateVertexInstances( const TArray<FVertexInstanceToCreate>& VertexInstancesToCreate, TArray<FVertexInstanceID>& OutNewVertexInstanceIDs )
+{
+	UE_LOG( LogEditableMesh, Verbose, TEXT( "CreateVertexInstances: %s" ), *LogHelpers::ArrayToString( VertexInstancesToCreate ) );
+
+	OutNewVertexInstanceIDs.Reset( VertexInstancesToCreate.Num() );
+
+	// Reserve elements
+	VertexInstances.Reserve( VertexInstances.Num() + VertexInstancesToCreate.Num() );
+
+	for( const FVertexInstanceToCreate& VertexInstanceToCreate : VertexInstancesToCreate )
+	{
+		// Allocate vertex instance
+		FVertexInstanceID VertexInstanceID = VertexInstanceToCreate.OriginalVertexInstanceID;
+		if( VertexInstanceID != FVertexInstanceID::Invalid )
+		{
+			VertexInstances.Insert( VertexInstanceID.GetValue(), FMeshVertexInstance() );
+		}
+		else
+		{
+			VertexInstanceID = FVertexInstanceID( VertexInstances.Add( FMeshVertexInstance() ) );
+		}
+
+		FMeshVertexInstance& VertexInstance = VertexInstances[ VertexInstanceID.GetValue() ];
+
+		// Initialize vertex instance
+		VertexInstance.VertexID = VertexInstanceToCreate.VertexID;
+
+		VertexInstance.Normal = FVector::ZeroVector;
+		VertexInstance.Tangent = FVector::ZeroVector;
+		VertexInstance.BinormalSign = 0.0f;
+		VertexInstance.Color = FLinearColor::White;
+		VertexInstance.VertexUVs.SetNumZeroed( TextureCoordinateCount );
+
+		// Add to the list of instances in the vertex
+		check( !Vertices[ VertexInstance.VertexID.GetValue() ].VertexInstanceIDs.Contains( VertexInstanceID ) );
+		Vertices[ VertexInstance.VertexID.GetValue() ].VertexInstanceIDs.Add( VertexInstanceID );
+
+		OutNewVertexInstanceIDs.Add( VertexInstanceID );
+	}
+
+	// NOTE: We iterate backwards, to delete vertex instances in the opposite order that we added them
+	{
+		FDeleteVertexInstancesChangeInput RevertInput;
+		RevertInput.bDeleteOrphanedVertices = false;
+		RevertInput.VertexInstanceIDsToDelete.Reserve( OutNewVertexInstanceIDs.Num() );
+		for( int32 VertexNumber = OutNewVertexInstanceIDs.Num() - 1; VertexNumber >= 0; --VertexNumber )
+		{
+			RevertInput.VertexInstanceIDsToDelete.Add( OutNewVertexInstanceIDs[ VertexNumber ] );
+		}
+
+		AddUndo( MakeUnique<FDeleteVertexInstancesChange>( MoveTemp( RevertInput ) ) );
+	}
+
+	// Advise the adapter of new vertex instances
+	for( UEditableMeshAdapter* Adapter : Adapters )
+	{
+		Adapter->OnCreateVertexInstances( this, OutNewVertexInstanceIDs );
+	}
+
+	// Now set vertex instance attributes
+	for( int32 Index = 0; Index < OutNewVertexInstanceIDs.Num(); ++Index )
+	{
+		for( const FMeshElementAttributeData& VertexInstanceAttribute : VertexInstancesToCreate[ Index ].VertexInstanceAttributes.Attributes )
+		{
+			SetVertexInstanceAttribute( OutNewVertexInstanceIDs[ Index ], VertexInstanceAttribute.AttributeName, VertexInstanceAttribute.AttributeIndex, VertexInstanceAttribute.AttributeValue );
+		}
+	}
+
+	UE_LOG( LogEditableMesh, Verbose, TEXT( "* CreateVertexInstances returned %s" ), *LogHelpers::ArrayToString( OutNewVertexInstanceIDs ) );
+}
+
+
 void UEditableMesh::CreateEdges( const TArray<FEdgeToCreate>& EdgesToCreate, TArray<FEdgeID>& OutNewEdgeIDs )
 {
-	OutNewEdgeIDs.Reset();
+	UE_LOG( LogEditableMesh, Verbose, TEXT( "CreateEdges: %s" ), *LogHelpers::ArrayToString( EdgesToCreate ) );
 
+	OutNewEdgeIDs.Reset( EdgesToCreate.Num() );
+
+	// Reserve elements
+	Edges.Reserve( Edges.Num() + EdgesToCreate.Num() );
+
 	for( const FEdgeToCreate& EdgeToCreate : EdgesToCreate )
 	{
-		// Create the edge and set it's attributes
-		FEdgeID NewEdgeID = FEdgeID::Invalid;
-		CreateEdge_Internal( EdgeToCreate.VertexID0, EdgeToCreate.VertexID1, EdgeToCreate.ConnectedPolygons, EdgeToCreate.OriginalEdgeID, /* Out */ NewEdgeID );
+		// Allocate edge
+		FEdgeID EdgeID = EdgeToCreate.OriginalEdgeID;
+		if( EdgeID != FEdgeID::Invalid )
+		{
+			Edges.Insert( EdgeID.GetValue(), FMeshEdge() );
+		}
+		else
+		{
+			EdgeID = FEdgeID( Edges.Add( FMeshEdge() ) );
+		}
 
-		OutNewEdgeIDs.Add( NewEdgeID );
+		FMeshEdge& Edge = Edges[ EdgeID.GetValue() ];
 
-		for( const FMeshElementAttributeData& EdgeAttribute : EdgeToCreate.EdgeAttributes.Attributes )
-		{
-			SetEdgeAttribute_Internal( NewEdgeID, EdgeAttribute.AttributeName, EdgeAttribute.AttributeIndex, EdgeAttribute.AttributeValue );
-		}
+		// Initialize edge
+		Edge.VertexIDs[ 0 ] = EdgeToCreate.VertexID0;
+		Edge.VertexIDs[ 1 ] = EdgeToCreate.VertexID1;
+		Edge.ConnectedPolygons = EdgeToCreate.ConnectedPolygons;
+
+		Edge.bIsHardEdge = false;
+		Edge.CreaseSharpness = 0.0f;
+
+		// Connect the edge to its vertices
+		Vertices[ Edge.VertexIDs[ 0 ].GetValue() ].ConnectedEdgeIDs.AddUnique( EdgeID );
+		Vertices[ Edge.VertexIDs[ 1 ].GetValue() ].ConnectedEdgeIDs.AddUnique( EdgeID );
+
+		OutNewEdgeIDs.Add( EdgeID );
 	}
 
 	// NOTE: We iterate backwards, to delete edges in the opposite order that we added them
-	FDeleteEdgesChangeInput RevertInput;
-	RevertInput.bDeleteOrphanedVertices = false;	// Don't delete any vertices on revert.  We're only creating edges here, not vertices!
-	RevertInput.EdgeIDsToDelete.Reserve( EdgesToCreate.Num() );
-	for( int32 EdgeNumber = OutNewEdgeIDs.Num() - 1; EdgeNumber >= 0; --EdgeNumber )
 	{
-		RevertInput.EdgeIDsToDelete.Add( OutNewEdgeIDs[ EdgeNumber ] );
+		FDeleteEdgesChangeInput RevertInput;
+		RevertInput.bDeleteOrphanedVertices = false;	// Don't delete any vertices on revert.  We're only creating edges here, not vertices!
+		RevertInput.EdgeIDsToDelete.Reserve( EdgesToCreate.Num() );
+		for( int32 EdgeNumber = OutNewEdgeIDs.Num() - 1; EdgeNumber >= 0; --EdgeNumber )
+		{
+			RevertInput.EdgeIDsToDelete.Add( OutNewEdgeIDs[ EdgeNumber ] );
+		}
+
+		AddUndo( MakeUnique<FDeleteEdgesChange>( MoveTemp( RevertInput ) ) );
 	}
 
-	AddUndo( MakeUnique<FDeleteEdgesChange>( MoveTemp( RevertInput ) ) );
+	// Advise the adapter that edges have been created
+	for( UEditableMeshAdapter* Adapter : Adapters )
+	{
+		Adapter->OnCreateEdges( this, OutNewEdgeIDs );
+	}
+
+	for( int32 Index = 0; Index < OutNewEdgeIDs.Num(); ++Index )
+	{
+		for( const FMeshElementAttributeData& EdgeAttribute : EdgesToCreate[ Index ].EdgeAttributes.Attributes )
+		{
+			const bool bIsUndo = false;
+			SetEdgeAttribute( OutNewEdgeIDs[ Index ], EdgeAttribute.AttributeName, EdgeAttribute.AttributeIndex, EdgeAttribute.AttributeValue, bIsUndo );
+		}
+	}
+
+	UE_LOG( LogEditableMesh, Verbose, TEXT( "* CreateEdges returned %s" ), *LogHelpers::ArrayToString( OutNewEdgeIDs ) );
 }
 
 
-void UEditableMesh::CreatePolygons( const TArray<FPolygonToCreate>& PolygonsToCreate, TArray<FPolygonRef>& OutNewPolygonRefs, TArray<FEdgeID>& OutNewEdgeIDs )
+FVertexInstanceID UEditableMesh::CreateVertexInstanceForContourVertex( const FVertexAndAttributes& ContourVertex, const FPolygonID PolygonID )
 {
-	OutNewPolygonRefs.Reset();
+	FVertexInstanceID NewVertexInstanceID = ContourVertex.VertexInstanceID;
+
+	// Determine whether we need to create a new vertex instance or not. We need to if:
+	//   a) no valid vertex instance ID was supplied;
+	// and then:
+	//   b) if there are no vertex instances in the same soft edged polygon group as the new polygon which share the attributes being assigned
+	if( NewVertexInstanceID == FVertexInstanceID::Invalid )
+	{
+		const FVertexID VertexID = ContourVertex.VertexID;
+		check( VertexID != FVertexID::Invalid );
+
+		// Find candidate vertex instances in the same smoothing group as this polygon
+		static TArray<FVertexInstanceID> PossibleVertexInstances;
+		PossibleVertexInstances.Reset();
+
+		const bool bPolygonNotYetInitialized = true;
+		GetVertexInstancesInSameSoftEdgedGroup( VertexID, PolygonID, bPolygonNotYetInitialized, PossibleVertexInstances );
+
+		// Iterate through each potential vertex instance, checking if its attributes are equal before and after the new attributes are applied.
+		// If so, we can merge with this vertex.
+		// Note: unspecified attributes in the FVertexAndAttributes struct are "don't cares".
+		for( const FVertexInstanceID PossibleVertexInstance : PossibleVertexInstances )
+		{
+			// Get mergeable attribute data for this vertex instance
+			static TArray<FMeshElementAttributeData> AllAttributes;
+			GetMergedNamedVertexInstanceAttributeData( GetMergeableVertexInstanceAttributes(), PossibleVertexInstance, TArray<FMeshElementAttributeData>(), AllAttributes );
+
+			// Get attribute data for this vertex instance as it would be after applying the new attributes
+			static TArray<FMeshElementAttributeData> AllAttributesModified;
+			GetMergedNamedVertexInstanceAttributeData( GetMergeableVertexInstanceAttributes(), PossibleVertexInstance, ContourVertex.PolygonVertexAttributes.Attributes, AllAttributesModified );
+
+			// If it's the same, we can use this vertex instance instead of creating a new one
+			if( AreAttributeListsEqual( AllAttributes, AllAttributesModified ) )
+			{
+				NewVertexInstanceID = PossibleVertexInstance;
+				break;
+			}
+		}
+
+		// If we didn't find any suitable existing vertex instance to merge with, create a new one.
+		if( NewVertexInstanceID == FVertexInstanceID::Invalid )
+		{
+			FVertexInstanceToCreate VertexInstanceToCreate;
+			VertexInstanceToCreate.VertexID = VertexID;
+			VertexInstanceToCreate.VertexInstanceAttributes = ContourVertex.PolygonVertexAttributes;
+
+			static TArray<FVertexInstanceToCreate> VertexInstancesToCreate;
+			VertexInstancesToCreate.Reset();
+			VertexInstancesToCreate.Add( VertexInstanceToCreate );
+
+			static TArray<FVertexInstanceID> CreatedVertexInstances;
+			CreateVertexInstances( VertexInstancesToCreate, CreatedVertexInstances );
+			check( CreatedVertexInstances.Num() == 1 );
+			NewVertexInstanceID = CreatedVertexInstances[ 0 ];
+		}
+	}
+	else
+	{
+		// Cannot specify both a vertex instance ID and a vertex ID
+		check( ContourVertex.VertexID == FVertexID::Invalid );
+	}
+
+	return NewVertexInstanceID;
+}
+
+
+void UEditableMesh::CreatePolygonContour( const FPolygonID PolygonID, const TArray<FVertexAndAttributes>& Contour, const EPolygonEdgeHardness PolygonEdgeHardness, TArray<FEdgeID>& OutEdgeIDs, TArray<FVertexInstanceID>& OutVertexInstanceIDs )
+{
+	// All polygons must have at least three vertices
+	const int32 NumContourVertices = Contour.Num();
+	check( NumContourVertices >= 3 );
+
+	OutEdgeIDs.Reset( NumContourVertices );
+	OutVertexInstanceIDs.SetNumUninitialized( NumContourVertices );
+
+	static TArray<FEdgeID> AllEdges;
+	AllEdges.SetNumUninitialized( NumContourVertices );
+
+	// Get or create edges
+	for( int32 VertexNumber = 0; VertexNumber < NumContourVertices; ++VertexNumber )
+	{
+		const int32 NextVertexNumber = ( VertexNumber + 1 ) % NumContourVertices;
+
+		// Get vertex ID for the ends of the edge.
+		FVertexID VertexID0 = Contour[ VertexNumber ].VertexID;
+		FVertexID VertexID1 = Contour[ NextVertexNumber ].VertexID;
+
+		// If valid vertex IDs weren't supplied, we'll deduce it from the vertex instance IDs instead (which must be valid instead).
+		if( VertexID0 == FVertexID::Invalid )
+		{
+			const FVertexInstanceID VertexInstanceID0 = Contour[ VertexNumber ].VertexInstanceID;
+			check( VertexInstanceID0 != FVertexInstanceID::Invalid );
+			VertexID0 = VertexInstances[ VertexInstanceID0.GetValue() ].VertexID;
+		}
+
+		if( VertexID1 == FVertexID::Invalid )
+		{
+			const FVertexInstanceID VertexInstanceID1 = Contour[ NextVertexNumber ].VertexInstanceID;
+			check( VertexInstanceID1 != FVertexInstanceID::Invalid );
+			VertexID1 = VertexInstances[ VertexInstanceID1.GetValue() ].VertexID;
+		}
+
+		bool bOutEdgeWindingIsReversed;
+		FEdgeID EdgeID = GetVertexPairEdge( VertexID0, VertexID1, bOutEdgeWindingIsReversed );
+		if( EdgeID == FEdgeID::Invalid )
+		{
+			// Create the new edge!  Note that this does not connect the edge to the polygon.  We expect the caller to do that afterwards.
+			FEdgeToCreate EdgeToCreate;
+			EdgeToCreate.VertexID0 = VertexID0;
+			EdgeToCreate.VertexID1 = VertexID1;
+
+			if( PolygonEdgeHardness == EPolygonEdgeHardness::NewEdgesHard )
+			{
+				EdgeToCreate.EdgeAttributes.Attributes.Emplace( UEditableMeshAttribute::EdgeIsHard(), 0, FVector4( 1.0f ) );
+			}
+
+			static TArray<FEdgeToCreate> EdgesToCreate;
+			EdgesToCreate.Reset();
+			EdgesToCreate.Add( EdgeToCreate );
+
+			static TArray<FEdgeID> NewEdgeIDs;
+			CreateEdges( EdgesToCreate, /* Out */ NewEdgeIDs );
+			check( NewEdgeIDs.Num() == 1 );
+			EdgeID = NewEdgeIDs[ 0 ];
+
+			OutEdgeIDs.Add( EdgeID );
+		}
+
+		// Connect the edge to the polygon
+		FMeshEdge& Edge = Edges[ EdgeID.GetValue() ];
+		Edge.ConnectedPolygons.Add( PolygonID );
+
+		AllEdges[ VertexNumber ] = EdgeID;
+	}
+
+	if( PolygonEdgeHardness == EPolygonEdgeHardness::AllEdgesHard ||
+		PolygonEdgeHardness == EPolygonEdgeHardness::AllEdgesSoft )
+	{
+		static TArray<FAttributesForEdge> AttributesForEdges;
+		AttributesForEdges.Reset( AllEdges.Num() );
+
+		for( const FEdgeID EdgeID : AllEdges )
+		{
+			AttributesForEdges.Emplace();
+			FAttributesForEdge& AttributesForEdge = AttributesForEdges.Last();
+			AttributesForEdge.EdgeID = EdgeID;
+			AttributesForEdge.EdgeAttributes.Attributes.Emplace( UEditableMeshAttribute::EdgeIsHard(), 0, FVector4( PolygonEdgeHardness == EPolygonEdgeHardness::AllEdgesHard ? 1.0f : 0.0f ) );
+		}
+
+		SetEdgesAttributes( AttributesForEdges );
+	}
+
+	// Get or create new vertex instances for this polygon.
+	// Do this after creating all the edges because:
+	//  1) we can establish a smoothing group for the new polygon based on adjacent edge hardness.
+	//  2) the revert operation will first delete all vertex instances, and then the edges (which can only delete orphaned vertices if they have no instances).
+	for( int32 VertexNumber = 0; VertexNumber < NumContourVertices; ++VertexNumber )
+	{
+		const FVertexInstanceID ContourVertexInstanceID = CreateVertexInstanceForContourVertex( Contour[ VertexNumber ], PolygonID );
+		OutVertexInstanceIDs[ VertexNumber ] = ContourVertexInstanceID;
+
+		// Connect this vertex instance to the polygon
+		FMeshVertexInstance& ContourVertexInstance = VertexInstances[ ContourVertexInstanceID.GetValue() ];
+		check( !ContourVertexInstance.ConnectedPolygons.Contains( PolygonID ) );
+		ContourVertexInstance.ConnectedPolygons.Add( PolygonID );
+	}
+}
+
+
+void UEditableMesh::CreatePolygons( const TArray<FPolygonToCreate>& PolygonsToCreate, TArray<FPolygonID>& OutNewPolygonIDs, TArray<FEdgeID>& OutNewEdgeIDs )
+{
+	UE_LOG( LogEditableMesh, Verbose, TEXT( "CreatePolygons: %s" ), *LogHelpers::ArrayToString( PolygonsToCreate ) );
+
+	OutNewPolygonIDs.Reset( PolygonsToCreate.Num() );
 	OutNewEdgeIDs.Reset();
 
+	Polygons.Reserve( Polygons.Num() + PolygonsToCreate.Num() );
+
 	for( const FPolygonToCreate& PolygonToCreate : PolygonsToCreate )
 	{
-		static TArray<FVertexID> PerimeterVertexIDs;
+		// Allocate polygon
+		FPolygonID PolygonID = PolygonToCreate.OriginalPolygonID;
+		if( PolygonID != FPolygonID::Invalid )
 		{
-			PerimeterVertexIDs.Reset();
-			PerimeterVertexIDs.Reserve( PolygonToCreate.PerimeterVertices.Num() );
-			for( const FVertexAndAttributes& PerimeterVertex : PolygonToCreate.PerimeterVertices )
-			{
-				PerimeterVertexIDs.Add( PerimeterVertex.VertexID );
-			}
+			Polygons.Insert( PolygonID.GetValue(), FMeshPolygon() );
+		}
+		else
+		{
+			PolygonID = FPolygonID( Polygons.Add( FMeshPolygon() ) );
 		}
 
-		static TArray<TArray<FVertexID>> VertexIDsForEachHole;
+		// Connect polygon to group and vice versa
+		FMeshPolygon& Polygon = Polygons[ PolygonID.GetValue() ];
+		FMeshPolygonGroup& PolygonGroup = PolygonGroups[ PolygonToCreate.PolygonGroupID.GetValue() ];
+		Polygon.PolygonGroupID = PolygonToCreate.PolygonGroupID;
+		PolygonGroup.Polygons.Add( PolygonID );
+
+		OutNewPolygonIDs.Add( PolygonID );
+
+		// Create edges and vertex instances for perimeter contour
+		static TArray<FEdgeID> PerimeterEdgeIDs;
+		CreatePolygonContour( PolygonID, PolygonToCreate.PerimeterVertices, PolygonToCreate.PolygonEdgeHardness, PerimeterEdgeIDs, Polygon.PerimeterContour.VertexInstanceIDs );
+		OutNewEdgeIDs.Append( PerimeterEdgeIDs );
+
+		// Iterate through hole definitions, and create edges and vertex instances for those
+		for( const FPolygonHoleVertices& PolygonHole : PolygonToCreate.PolygonHoles )
 		{
-			VertexIDsForEachHole.Reset();
-			VertexIDsForEachHole.Reserve( PolygonToCreate.PolygonHoles.Num() );
+			Polygon.HoleContours.Emplace();
+			FMeshPolygonContour& Contour = Polygon.HoleContours.Last();
 
-			for( const FPolygonHoleVertices& PolygonHoleVertices : PolygonToCreate.PolygonHoles )
-			{
-				TArray<FVertexID>& HoleVertexIDs = *new( VertexIDsForEachHole ) TArray<FVertexID>();
-				HoleVertexIDs.Reserve( PolygonHoleVertices.HoleVertices.Num() );
-				for( const FVertexAndAttributes& HoleVertex : PolygonHoleVertices.HoleVertices )
-				{
-					HoleVertexIDs.Add( HoleVertex.VertexID );
-				}
-			}
-		}
-
-		// Create the polygon (along with any missing edges that are needed to connect the new polygon's vertices)
-		static TArray<FEdgeID> NewEdgeIDsForThisPolygon;
-		FPolygonRef NewPolygonRef;
-		CreatePolygon_Internal( PolygonToCreate.SectionID, PerimeterVertexIDs, VertexIDsForEachHole, PolygonToCreate.OriginalPolygonID, /* Out */ NewPolygonRef, /* Out */ NewEdgeIDsForThisPolygon );
-
-		OutNewEdgeIDs.Append( NewEdgeIDsForThisPolygon );
-		OutNewPolygonRefs.Add( NewPolygonRef );
-
-		// Set the polygon's vertex attributes
-		{
-			// Polygon perimeter
-			if( PolygonToCreate.PerimeterVertices.Num() > 0 )
-			{
-				const int32 PerimeterVertexCount = PolygonToCreate.PerimeterVertices.Num();
-				for( int32 PerimeterVertexNumber = 0; PerimeterVertexNumber < PerimeterVertexCount; ++PerimeterVertexNumber )
-				{
-					const TArray<FMeshElementAttributeData>& PerimeterVertexAttributeList = PolygonToCreate.PerimeterVertices[ PerimeterVertexNumber ].PolygonVertexAttributes.Attributes;
-					for( const FMeshElementAttributeData& PerimeterVertexAttribute : PerimeterVertexAttributeList )
-					{
-						SetPolygonPerimeterVertexAttribute_Internal( NewPolygonRef, PerimeterVertexNumber, PerimeterVertexAttribute.AttributeName, PerimeterVertexAttribute.AttributeIndex, PerimeterVertexAttribute.AttributeValue );
-					}
-				}
-			}
-
-			// Polygon holes
-			if( PolygonToCreate.PolygonHoles.Num() > 0 )
-			{
-				const int32 HoleCount = PolygonToCreate.PolygonHoles.Num();
-				for( int32 HoleNumber = 0; HoleNumber < HoleCount; ++HoleNumber )
-				{
-					const FPolygonHoleVertices& PolygonHoleVertices = PolygonToCreate.PolygonHoles[ HoleNumber ];
-					if( PolygonHoleVertices.HoleVertices.Num() > 0 )
-					{
-						const int32 HoleVertexCount = PolygonHoleVertices.HoleVertices.Num();
-						for( int32 HoleVertexNumber = 0; HoleVertexNumber < HoleVertexCount; ++HoleVertexNumber )
-						{
-							const TArray<FMeshElementAttributeData>& HoleVertexAttributeList = PolygonHoleVertices.HoleVertices[ HoleVertexNumber ].PolygonVertexAttributes.Attributes;
-							for( const FMeshElementAttributeData& HoleVertexAttribute : HoleVertexAttributeList )
-							{
-								SetPolygonHoleVertexAttribute_Internal( NewPolygonRef, HoleNumber, HoleVertexNumber, HoleVertexAttribute.AttributeName, HoleVertexAttribute.AttributeIndex, HoleVertexAttribute.AttributeValue );
-							}
-						}
-					}
-				}
-			}
+			static TArray<FEdgeID> HoleEdgeIDs;
+			CreatePolygonContour( PolygonID, PolygonHole.HoleVertices, EPolygonEdgeHardness::AllEdgesHard, HoleEdgeIDs, Contour.VertexInstanceIDs );
+			OutNewEdgeIDs.Append( HoleEdgeIDs );
 		}
 	}
 
+	for( UEditableMeshAdapter* Adapter : Adapters )
+	{
+		Adapter->OnCreatePolygons( this, OutNewPolygonIDs );
+	}
+
+	// Generate tangent basis for the polygon
+	PolygonsPendingNewTangentBasis.Append( OutNewPolygonIDs );
+
+	// Generate triangles for the new polygon
+	PolygonsPendingTriangulation.Append( OutNewPolygonIDs );
+
 	// NOTE: We iterate backwards, to delete polygons in the opposite order that we added them
-	FDeletePolygonsChangeInput RevertInput;
-	RevertInput.PolygonRefsToDelete.Reserve( PolygonsToCreate.Num() );
-	RevertInput.bDeleteOrphanedEdges = false;	// Don't delete any edges on revert.  We're only creating polygons here, not vertices!
-	RevertInput.bDeleteOrphanedVertices = false;	// Don't delete any vertices on revert.  We're only creating polygons here, not vertices!
-	RevertInput.bDeleteEmptySections = false;
-	for( int32 PolygonNumber = OutNewPolygonRefs.Num() - 1; PolygonNumber >= 0; --PolygonNumber )
 	{
-		RevertInput.PolygonRefsToDelete.Add( OutNewPolygonRefs[ PolygonNumber ] );
+		FDeletePolygonsChangeInput RevertInput;
+		RevertInput.PolygonIDsToDelete.Reserve( PolygonsToCreate.Num() );
+		RevertInput.bDeleteOrphanedEdges = false;	// Don't delete any edges on revert.  We're only creating polygons here, not vertices!
+		RevertInput.bDeleteOrphanedVertices = false;	// Don't delete any vertices on revert.  We're only creating polygons here, not vertices!
+		RevertInput.bDeleteOrphanedVertexInstances = false;
+		RevertInput.bDeleteEmptySections = false;
+		for( int32 PolygonNumber = OutNewPolygonIDs.Num() - 1; PolygonNumber >= 0; --PolygonNumber )
+		{
+			RevertInput.PolygonIDsToDelete.Add( OutNewPolygonIDs[ PolygonNumber ] );
+		}
+
+		AddUndo( MakeUnique<FDeletePolygonsChange>( MoveTemp( RevertInput ) ) );
 	}
 
-	AddUndo( MakeUnique<FDeletePolygonsChange>( MoveTemp( RevertInput ) ) );
+	UE_LOG( LogEditableMesh, Verbose, TEXT( "* CreatePolygons returned %s, %s" ), *LogHelpers::ArrayToString( OutNewPolygonIDs ), *LogHelpers::ArrayToString( OutNewEdgeIDs ) );
 }
 
 
-void UEditableMesh::DeletePolygons( const TArray<FPolygonRef>& PolygonRefsToDelete, const bool bDeleteOrphanedEdges, const bool bDeleteOrphanedVertices, const bool bDeleteEmptySections )
+void UEditableMesh::BackupPolygonContour( const FMeshPolygonContour& Contour, TArray<FVertexAndAttributes>& OutVerticesAndAttributes )
 {
-	// Back everything up
-	FCreatePolygonsChangeInput RevertInput;
+	OutVerticesAndAttributes.Reserve( Contour.VertexInstanceIDs.Num() );
+	for( const FVertexInstanceID VertexInstanceID : Contour.VertexInstanceIDs )
 	{
-		RevertInput.PolygonsToCreate.Reserve( PolygonRefsToDelete.Num() );
+		const FMeshVertexInstance& VertexInstance = VertexInstances[ VertexInstanceID.GetValue() ];
 
-		// NOTE: We iterate backwards, to restore edges in the opposite order that we deleted them
-		for( int32 PolygonNumber = PolygonRefsToDelete.Num() - 1; PolygonNumber >= 0; --PolygonNumber )
-		{
-			const FPolygonRef& PolygonRef = PolygonRefsToDelete[ PolygonNumber ];
+		OutVerticesAndAttributes.Emplace();
+		FVertexAndAttributes& VertexAndAttributes = OutVerticesAndAttributes.Last();
 
-			FPolygonToCreate& PolygonToCreate = *new( RevertInput.PolygonsToCreate ) FPolygonToCreate();
-			PolygonToCreate.SectionID = PolygonRef.SectionID;
-			PolygonToCreate.OriginalPolygonID = PolygonRef.PolygonID;
-
-			static TArray<FVertexID> PolygonPerimeterVertexIDs;
-			GetPolygonPerimeterVertices( PolygonRef, /* Out */ PolygonPerimeterVertexIDs );
-
-			PolygonToCreate.PerimeterVertices.Reserve( PolygonPerimeterVertexIDs.Num() );
-			for( int32 PerimeterVertexNumber = 0; PerimeterVertexNumber < PolygonPerimeterVertexIDs.Num(); ++PerimeterVertexNumber )
-			{
-				const FVertexID PerimeterVertexID = PolygonPerimeterVertexIDs[ PerimeterVertexNumber ];
-
-				FVertexAndAttributes& VertexAndAttributes = *new( PolygonToCreate.PerimeterVertices ) FVertexAndAttributes();
-				VertexAndAttributes.VertexID = PerimeterVertexID;
-
-				for( const FName AttributeName : UEditableMesh::GetValidPolygonVertexAttributes() )
-				{
-					const int32 MaxAttributeIndex = GetMaxAttributeIndex( AttributeName );
-					for( int32 AttributeIndex = 0; AttributeIndex < MaxAttributeIndex; ++AttributeIndex )
-					{
-						FMeshElementAttributeData& PolygonVertexAttribute = *new( VertexAndAttributes.PolygonVertexAttributes.Attributes ) FMeshElementAttributeData(
-							AttributeName,
-							AttributeIndex,
-							GetPolygonPerimeterVertexAttribute( PolygonRef, PerimeterVertexNumber, AttributeName, AttributeIndex ) );
-					}
-				}
-			}
-
-			const int32 HoleCount = GetPolygonHoleCount( PolygonRef );
-			PolygonToCreate.PolygonHoles.SetNum( HoleCount, false );
-			for( int32 HoleNumber = 0; HoleNumber < HoleCount; ++HoleNumber )
-			{
-				FPolygonHoleVertices& Hole = PolygonToCreate.PolygonHoles[ HoleNumber ];
-
-				static TArray<FVertexID> PolygonHoleVertexIDs;
-				GetPolygonHoleVertices( PolygonRef, HoleNumber, /* Out */ PolygonHoleVertexIDs );
-
-				Hole.HoleVertices.Reserve( PolygonHoleVertexIDs.Num() );
-				for( int32 HoleVertexNumber = 0; HoleVertexNumber < PolygonHoleVertexIDs.Num(); ++HoleVertexNumber )
-				{
-					const FVertexID HoleVertexID = PolygonHoleVertexIDs[ HoleVertexNumber ];
-
-					FVertexAndAttributes& VertexAndAttributes = *new( Hole.HoleVertices ) FVertexAndAttributes();
-					VertexAndAttributes.VertexID = HoleVertexID;
-
-					for( const FName AttributeName : UEditableMesh::GetValidPolygonVertexAttributes() )
-					{
-						const int32 MaxAttributeIndex = GetMaxAttributeIndex( AttributeName );
-						for( int32 AttributeIndex = 0; AttributeIndex < MaxAttributeIndex; ++AttributeIndex )
-						{
-							FMeshElementAttributeData& PolygonVertexAttribute = *new( VertexAndAttributes.PolygonVertexAttributes.Attributes ) FMeshElementAttributeData(
-								AttributeName,
-								AttributeIndex,
-								GetPolygonHoleVertexAttribute( PolygonRef, HoleNumber, HoleVertexNumber, AttributeName, AttributeIndex ) );
-						}
-					}
-				}
-			}
-		}
-	}
-
-	AddUndo( MakeUnique<FCreatePolygonsChange>( MoveTemp( RevertInput ) ) );
-
-	for( const FPolygonRef& PolygonRefToDelete : PolygonRefsToDelete )
-	{
-		DeletePolygon_Internal( PolygonRefToDelete, bDeleteOrphanedEdges, bDeleteOrphanedVertices, bDeleteEmptySections );
+		// We rely on undoing recreating vertex instances, therefore we only need pass their IDs.
+		VertexAndAttributes.VertexInstanceID = VertexInstanceID;
+		VertexAndAttributes.VertexID = FVertexID::Invalid;
 	}
 }
 
 
+void UEditableMesh::DeletePolygonContour( const FPolygonID PolygonID, const FMeshPolygonContour& Contour, const bool bDeleteOrphanedEdges, const bool bDeleteOrphanedVertexInstances, TArray<FEdgeID>& OrphanedEdgeIDs, TArray<FVertexInstanceID>& OrphanedVertexInstanceIDs )
+{
+	for( const FVertexInstanceID VertexInstanceID : Contour.VertexInstanceIDs )
+	{
+		// Disconnect the polygon from the vertex instance
+		FMeshVertexInstance& VertexInstance = VertexInstances[ VertexInstanceID.GetValue() ];
+		verify( VertexInstance.ConnectedPolygons.Remove( PolygonID ) == 1 );
+
+		// Vertex instances which are only part of this contour (i.e. not shared by other polygons) are optionally deleted
+		if( bDeleteOrphanedVertexInstances && VertexInstance.ConnectedPolygons.Num() == 0 )
+		{
+			OrphanedVertexInstanceIDs.Add( VertexInstanceID );
+		}
+	}
+
+	// Remove reference to polygon from edges
+	static TArray<FEdgeID> EdgeIDs;
+	GetPolygonContourEdges( Contour, EdgeIDs );
+
+	for( const FEdgeID EdgeID : EdgeIDs )
+	{
+		FMeshEdge& Edge = Edges[ EdgeID.GetValue() ];
+		verify( Edge.ConnectedPolygons.Remove( PolygonID ) == 1 );
+
+		// Identify orphaned edges
+		if( bDeleteOrphanedEdges && Edge.ConnectedPolygons.Num() == 0 )
+		{
+			OrphanedEdgeIDs.Add( EdgeID );
+		}
+	}
+}
+
+
+void UEditableMesh::DeletePolygons( const TArray<FPolygonID>& PolygonIDsToDelete, const bool bDeleteOrphanedEdges, const bool bDeleteOrphanedVertices, const bool bDeleteOrphanedVertexInstances, const bool bDeleteEmptyPolygonGroups )
+{
+	UE_LOG( LogEditableMesh, Verbose, TEXT( "DeletePolygons: %s" ), *LogHelpers::ArrayToString( PolygonIDsToDelete ) );
+
+	// Back everything up
+	{
+		FCreatePolygonsChangeInput RevertInput;
+		RevertInput.PolygonsToCreate.Reserve( PolygonIDsToDelete.Num() );
+
+		// NOTE: We iterate backwards, to restore edges in the opposite order that we deleted them
+		for( int32 PolygonNumber = PolygonIDsToDelete.Num() - 1; PolygonNumber >= 0; --PolygonNumber )
+		{
+			const FPolygonID PolygonID = PolygonIDsToDelete[ PolygonNumber ];
+			const FMeshPolygon& Polygon = Polygons[ PolygonID.GetValue() ];
+
+			RevertInput.PolygonsToCreate.Emplace();
+			FPolygonToCreate& PolygonToCreate = RevertInput.PolygonsToCreate.Last();
+
+			PolygonToCreate.PolygonGroupID = Polygon.PolygonGroupID;
+			PolygonToCreate.OriginalPolygonID = PolygonID;
+
+			BackupPolygonContour( Polygon.PerimeterContour, PolygonToCreate.PerimeterVertices );
+			
+			PolygonToCreate.PolygonHoles.Reserve( Polygon.HoleContours.Num() );
+			for( const FMeshPolygonContour& HoleContour : Polygon.HoleContours )
+			{
+				PolygonToCreate.PolygonHoles.Emplace();
+				BackupPolygonContour( HoleContour, PolygonToCreate.PolygonHoles.Last().HoleVertices );
+			}
+		}
+
+		AddUndo( MakeUnique<FCreatePolygonsChange>( MoveTemp( RevertInput ) ) );
+	}
+
+	for( UEditableMeshAdapter* Adapter : Adapters )
+	{
+		Adapter->OnDeletePolygons( this, PolygonIDsToDelete );
+	}
+
+	// Delete the polygons
+	{
+		static TArray<FEdgeID> OrphanedEdgeIDs;
+		OrphanedEdgeIDs.Reset();
+
+		static TArray<FVertexInstanceID> OrphanedVertexInstanceIDs;
+		OrphanedVertexInstanceIDs.Reset();
+
+		static TArray<FPolygonGroupID> EmptyPolygonGroupIDs;
+		EmptyPolygonGroupIDs.Reset();
+
+		for( const FPolygonID PolygonID : PolygonIDsToDelete )
+		{
+			const FMeshPolygon& Polygon = Polygons[ PolygonID.GetValue() ];
+
+			DeletePolygonContour( PolygonID, Polygon.PerimeterContour, bDeleteOrphanedEdges, bDeleteOrphanedVertexInstances, OrphanedEdgeIDs, OrphanedVertexInstanceIDs );
+
+			for( const FMeshPolygonContour& HoleContour : Polygon.HoleContours )
+			{
+				DeletePolygonContour( PolygonID, HoleContour, bDeleteOrphanedEdges, bDeleteOrphanedVertexInstances, OrphanedEdgeIDs, OrphanedVertexInstanceIDs );
+			}
+
+			// Disconnect the polygon from the polygon group
+			const FPolygonGroupID PolygonGroupID = Polygon.PolygonGroupID;
+			FMeshPolygonGroup& PolygonGroup = PolygonGroups[ PolygonGroupID.GetValue() ];
+			verify( PolygonGroup.Polygons.Remove( PolygonID ) == 1 );
+
+			if( bDeleteEmptyPolygonGroups && PolygonGroup.Polygons.Num() == 0 )
+			{
+				EmptyPolygonGroupIDs.Add( PolygonGroupID );
+			}
+
+			// Delete the polygon
+			Polygons.RemoveAt( PolygonID.GetValue() );
+		}
+
+		// Remove vertex instances which are exclusively used by this polygon.
+		// We do not want this to remove orphaned vertices; this will optionally happen below when removing edges.
+		if( OrphanedVertexInstanceIDs.Num() > 0 )
+		{
+			DeleteVertexInstances( OrphanedVertexInstanceIDs, false );
+		}
+
+		// Remove any edges which may have been orphaned. This may also optionally remove any orphaned vertices.
+		// We can do this here because we know any edges which were orphaned will have had only a single vertex instance at each vertex.
+		// Therefore the vertex will now have no instances further to deleting them above.
+		// Note: there is never a situation where there could be orphaned vertices but not orphaned edges.
+		if( OrphanedEdgeIDs.Num() > 0 )
+		{
+			DeleteEdges( OrphanedEdgeIDs, bDeleteOrphanedVertices );
+		}
+
+		// Remove any empty polygon groups which may have resulted
+		if( EmptyPolygonGroupIDs.Num() > 0 )
+		{
+			DeletePolygonGroups( EmptyPolygonGroupIDs );
+		}
+	}
+
+	// If any of these polygons are in the pending list for triangulation or computing a new tangent basis, remove them
+	for( const FPolygonID PolygonID : PolygonIDsToDelete )
+	{
+		PolygonsPendingNewTangentBasis.Remove( PolygonID );
+		PolygonsPendingTriangulation.Remove( PolygonID );
+	}
+
+	UE_LOG( LogEditableMesh, Verbose, TEXT( "* DeletePolygons returned" ) );
+}
+
+
+void UEditableMesh::CreatePolygonGroups( const TArray<FPolygonGroupToCreate>& PolygonGroupsToCreate, TArray<FPolygonGroupID>& OutNewPolygonGroupIDs )
+{
+	// @todo mesheditor: Material, bEnableCollision and bCastShadow should all be implemented as PolygonGroup attributes
+
+	OutNewPolygonGroupIDs.Reset( PolygonGroupsToCreate.Num() );
+
+	// Reserve elements
+	PolygonGroups.Reserve( PolygonGroups.Num() + PolygonGroupsToCreate.Num() );
+
+	for( const FPolygonGroupToCreate& PolygonGroupToCreate : PolygonGroupsToCreate )
+	{
+		// Allocate polygon group
+		FPolygonGroupID PolygonGroupID = PolygonGroupToCreate.OriginalPolygonGroupID;
+		if( PolygonGroupID != FPolygonGroupID::Invalid )
+		{
+			PolygonGroups.Insert( PolygonGroupID.GetValue(), FMeshPolygonGroup() );
+		}
+		else
+		{
+			PolygonGroupID = FPolygonGroupID( PolygonGroups.Add( FMeshPolygonGroup() ) );
+		}
+
+		FMeshPolygonGroup& PolygonGroup = PolygonGroups[ PolygonGroupID.GetValue() ];
+
+		PolygonGroup.Material = PolygonGroupToCreate.Material;
+		PolygonGroup.bEnableCollision = PolygonGroupToCreate.bEnableCollision;
+		PolygonGroup.bCastShadow = PolygonGroupToCreate.bCastShadow;
+
+		OutNewPolygonGroupIDs.Add( PolygonGroupID );
+	}
+
+	for( UEditableMeshAdapter* Adapter : Adapters )
+	{
+		Adapter->OnCreatePolygonGroups( this, OutNewPolygonGroupIDs );
+	}
+
+	// Back up
+	{
+		FDeletePolygonGroupsChangeInput RevertInput;
+		RevertInput.PolygonGroupIDs.Reserve( OutNewPolygonGroupIDs.Num() );
+		for( int32 Index = OutNewPolygonGroupIDs.Num() - 1; Index >= 0; --Index )
+		{
+			RevertInput.PolygonGroupIDs.Add( OutNewPolygonGroupIDs[ Index ] );
+		}
+
+		AddUndo( MakeUnique<FDeletePolygonGroupsChange>( MoveTemp( RevertInput ) ) );
+	}
+}
+
+
+void UEditableMesh::DeletePolygonGroups( const TArray<FPolygonGroupID>& PolygonGroupIDs )
+{
+	// Back everything up
+	{
+		FCreatePolygonGroupsChangeInput RevertInput;
+
+		for( int32 Index = PolygonGroupIDs.Num() - 1; Index >= 0; --Index )
+		{
+			const FPolygonGroupID PolygonGroupID = PolygonGroupIDs[ Index ];
+			const FMeshPolygonGroup& PolygonGroup = PolygonGroups[ PolygonGroupID.GetValue() ];
+
+			RevertInput.PolygonGroupsToCreate.Emplace();
+			FPolygonGroupToCreate& PolygonGroupToCreate = RevertInput.PolygonGroupsToCreate.Last();
+
+			PolygonGroupToCreate.OriginalPolygonGroupID = PolygonGroupID;
+			PolygonGroupToCreate.Material = PolygonGroup.Material;
+			PolygonGroupToCreate.bCastShadow = PolygonGroup.bCastShadow;
+			PolygonGroupToCreate.bEnableCollision = PolygonGroup.bEnableCollision;
+		}
+
+		AddUndo( MakeUnique<FCreatePolygonGroupsChange>( MoveTemp( RevertInput ) ) );
+	}
+
+	for( UEditableMeshAdapter* Adapter : Adapters )
+	{
+		Adapter->OnDeletePolygonGroups( this, PolygonGroupIDs );
+	}
+
+	// Delete the polygon groups
+	{
+		for( const FPolygonGroupID PolygonGroupID : PolygonGroupIDs )
+		{
+			const FMeshPolygonGroup& PolygonGroup = PolygonGroups[ PolygonGroupID.GetValue() ];
+			check( PolygonGroup.Polygons.Num() == 0 );
+
+			PolygonGroups.RemoveAt( PolygonGroupID.GetValue() );
+		}
+	}
+}
+
+
 void UEditableMesh::SetVerticesAttributes( const TArray<FAttributesForVertex>& AttributesForVertices )
 {
+	UE_LOG( LogEditableMesh, Verbose, TEXT( "SetVerticesAttributes: %s" ), *LogHelpers::ArrayToString( AttributesForVertices ) );
+
 	FSetVerticesAttributesChangeInput RevertInput;
 
 	RevertInput.AttributesForVertices.Reserve( AttributesForVertices.Num() );
@@ -3279,17 +4854,57 @@
 					GetVertexAttribute( AttributesForVertex.VertexID, VertexAttribute.AttributeName, VertexAttribute.AttributeIndex ) );
 
 				// Set the new attribute
-				SetVertexAttribute_Internal( AttributesForVertex.VertexID, VertexAttribute.AttributeName, VertexAttribute.AttributeIndex, VertexAttribute.AttributeValue );
+				SetVertexAttribute( AttributesForVertex.VertexID, VertexAttribute.AttributeName, VertexAttribute.AttributeIndex, VertexAttribute.AttributeValue );
 			}
 		}
 	}
 
 	AddUndo( MakeUnique<FSetVerticesAttributesChange>( MoveTemp( RevertInput ) ) );
+
+	UE_LOG( LogEditableMesh, Verbose, TEXT( "* SetVerticesAttributes returned" ) );
 }
 
 
+void UEditableMesh::SetVertexInstancesAttributes( const TArray<FAttributesForVertexInstance>& AttributesForVertexInstances )
+{
+	UE_LOG( LogEditableMesh, Verbose, TEXT( "SetVertexInstancesAttributes: %s" ), *LogHelpers::ArrayToString( AttributesForVertexInstances ) );
+
+	FSetVertexInstancesAttributesChangeInput RevertInput;
+
+	RevertInput.AttributesForVertexInstances.Reserve( AttributesForVertexInstances.Num() );
+	for( const FAttributesForVertexInstance& AttributesForVertexInstance : AttributesForVertexInstances )
+	{
+		FAttributesForVertexInstance& RevertVertexInstance = *new( RevertInput.AttributesForVertexInstances ) FAttributesForVertexInstance();
+		RevertVertexInstance.VertexInstanceID = AttributesForVertexInstance.VertexInstanceID;
+
+		// Set the vertex instance attributes
+		{
+			RevertVertexInstance.VertexInstanceAttributes.Attributes.Reserve( AttributesForVertexInstance.VertexInstanceAttributes.Attributes.Num() );
+
+			for( const FMeshElementAttributeData& VertexInstanceAttribute : AttributesForVertexInstance.VertexInstanceAttributes.Attributes )
+			{
+				// Back up the attribute
+				FMeshElementAttributeData& RevertAttribute = *new( RevertVertexInstance.VertexInstanceAttributes.Attributes ) FMeshElementAttributeData(
+					VertexInstanceAttribute.AttributeName,
+					VertexInstanceAttribute.AttributeIndex,
+					GetVertexInstanceAttribute( AttributesForVertexInstance.VertexInstanceID, VertexInstanceAttribute.AttributeName, VertexInstanceAttribute.AttributeIndex ) );
+
+				// Set the new attribute
+				SetVertexInstanceAttribute( AttributesForVertexInstance.VertexInstanceID, VertexInstanceAttribute.AttributeName, VertexInstanceAttribute.AttributeIndex, VertexInstanceAttribute.AttributeValue );
+			}
+		}
+	}
+
+	AddUndo( MakeUnique<FSetVertexInstancesAttributesChange>( MoveTemp( RevertInput ) ) );
+
+	UE_LOG( LogEditableMesh, Verbose, TEXT( "* SetVertexInstancesAttributes returned" ) );
+}
+
+
 void UEditableMesh::SetEdgesAttributes( const TArray<FAttributesForEdge>& AttributesForEdges )
 {
+	UE_LOG( LogEditableMesh, Verbose, TEXT( "SetEdgesAttributes: %s" ), *LogHelpers::ArrayToString( AttributesForEdges ) );
+
 	FSetEdgesAttributesChangeInput RevertInput;
 
 	RevertInput.AttributesForEdges.Reserve( AttributesForEdges.Num() );
@@ -3297,6 +4912,7 @@
 	{
 		FAttributesForEdge& RevertEdge = *new( RevertInput.AttributesForEdges ) FAttributesForEdge();
 		RevertEdge.EdgeID = AttributesForEdge.EdgeID;
+		RevertEdge.bIsUndo = true;
 
 		// Set the edge attributes
 		{
@@ -3311,103 +4927,475 @@
 					GetEdgeAttribute( AttributesForEdge.EdgeID, EdgeAttribute.AttributeName, EdgeAttribute.AttributeIndex ) );
 
 				// Set the new attribute
-				SetEdgeAttribute_Internal( AttributesForEdge.EdgeID, EdgeAttribute.AttributeName, EdgeAttribute.AttributeIndex, EdgeAttribute.AttributeValue );
+				SetEdgeAttribute( AttributesForEdge.EdgeID, EdgeAttribute.AttributeName, EdgeAttribute.AttributeIndex, EdgeAttribute.AttributeValue, AttributesForEdge.bIsUndo );
 			}
 		}
 	}
 
 	AddUndo( MakeUnique<FSetEdgesAttributesChange>( MoveTemp( RevertInput ) ) );
+
+	UE_LOG( LogEditableMesh, Verbose, TEXT( "* SetEdgesAttributes returned" ) );
 }
 
 
+const TArray<FName>& UEditableMesh::GetMergeableVertexInstanceAttributes()
+{
+	static TArray<FName> MergeableVertexInstanceAttributes;
+	if( MergeableVertexInstanceAttributes.Num() == 0 )
+	{
+		MergeableVertexInstanceAttributes.Add( UEditableMeshAttribute::VertexTextureCoordinate() );
+		MergeableVertexInstanceAttributes.Add( UEditableMeshAttribute::VertexColor() );
+	}
+	return MergeableVertexInstanceAttributes;
+}
+
+
+void UEditableMesh::GetMergedNamedVertexInstanceAttributeData( const TArray<FName>& AttributeNames, const FVertexInstanceID VertexInstanceID, const TArray<FMeshElementAttributeData>& AttributeValuesToMerge, TArray<FMeshElementAttributeData>& OutAttributes ) const
+{
+	OutAttributes.Reset();
+
+	for( const FName AttributeName : AttributeNames )
+	{
+		const int32 MaxAttributeIndex = GetMaxAttributeIndex( AttributeName );
+		for( int32 AttributeIndex = 0; AttributeIndex < MaxAttributeIndex; ++AttributeIndex )
+		{
+			const int32 ReplaceIndex =  AttributeValuesToMerge.IndexOfByPredicate(
+				[ AttributeName, AttributeIndex ]( const FMeshElementAttributeData& AttributeData )
+				{
+					return AttributeData.AttributeName == AttributeName && AttributeData.AttributeIndex == AttributeIndex;
+				}
+			);
+
+			if( ReplaceIndex == INDEX_NONE )
+			{
+				// If the named attribute and index doesn't appear in the list to merge, add the current attribute value to the list
+				OutAttributes.Emplace( AttributeName, AttributeIndex, GetVertexInstanceAttribute( VertexInstanceID, AttributeName, AttributeIndex ) );
+			}
+			else
+			{
+				// If the named attribute and index does appear in the list to merge, add that to the list instead
+				OutAttributes.Emplace( AttributeValuesToMerge[ ReplaceIndex ] );
+			}
+		}
+	}
+}
+
+
+void UEditableMesh::ChangePolygonsVertexInstances( const TArray<FChangeVertexInstancesForPolygon>& VertexInstancesForPolygons )
+{
+	UE_LOG( LogEditableMesh, Verbose, TEXT( "ChangePolygonsVertexInstances: %s" ), *LogHelpers::ArrayToString( VertexInstancesForPolygons ) );
+
+	// Back everything up
+	{
+		FChangePolygonsVertexInstancesChangeInput RevertInput;
+		RevertInput.VertexInstancesForPolygons.Reserve( VertexInstancesForPolygons.Num() );
+
+		// NOTE: We iterate backwards, to restore edges in the opposite order that we changed them
+		for( int32 Index = VertexInstancesForPolygons.Num() - 1; Index >= 0; --Index )
+		{
+			const FChangeVertexInstancesForPolygon& VertexInstancesForPolygon = VertexInstancesForPolygons[ Index ];
+			const FPolygonID PolygonID = VertexInstancesForPolygon.PolygonID;
+			const FMeshPolygon& Polygon = Polygons[ PolygonID.GetValue() ];
+
+			RevertInput.VertexInstancesForPolygons.Emplace();
+			FChangeVertexInstancesForPolygon& RevertVertexInstancesForPolygon = RevertInput.VertexInstancesForPolygons.Last();
+			RevertVertexInstancesForPolygon.PolygonID = PolygonID;
+
+			for( const FVertexIndexAndInstanceID& IndexAndInstance : VertexInstancesForPolygon.PerimeterVertexIndicesAndInstanceIDs )
+			{
+				RevertVertexInstancesForPolygon.PerimeterVertexIndicesAndInstanceIDs.Emplace();
+				FVertexIndexAndInstanceID& RevertIndexAndInstance = RevertVertexInstancesForPolygon.PerimeterVertexIndicesAndInstanceIDs.Last();
+				RevertIndexAndInstance.ContourIndex = IndexAndInstance.ContourIndex;
+				RevertIndexAndInstance.VertexInstanceID = Polygon.PerimeterContour.VertexInstanceIDs[ IndexAndInstance.ContourIndex ];
+			}
+
+			int32 HoleIndex = 0;
+			for( const FVertexInstancesForPolygonHole& VertexInstancesForPolygonHole : VertexInstancesForPolygon.VertexIndicesAndInstanceIDsForEachHole )
+			{
+				RevertVertexInstancesForPolygon.VertexIndicesAndInstanceIDsForEachHole.Emplace();
+				FVertexInstancesForPolygonHole& RevertVertexInstancesForPolygonHole = RevertVertexInstancesForPolygon.VertexIndicesAndInstanceIDsForEachHole.Last();
+
+				for( const FVertexIndexAndInstanceID& IndexAndInstance : VertexInstancesForPolygonHole.VertexIndicesAndInstanceIDs )
+				{
+					RevertVertexInstancesForPolygonHole.VertexIndicesAndInstanceIDs.Emplace();
+					FVertexIndexAndInstanceID& RevertIndexAndInstance = RevertVertexInstancesForPolygonHole.VertexIndicesAndInstanceIDs.Last();
+					RevertIndexAndInstance.ContourIndex = IndexAndInstance.ContourIndex;
+					RevertIndexAndInstance.VertexInstanceID = Polygon.HoleContours[ HoleIndex ].VertexInstanceIDs[ IndexAndInstance.ContourIndex ];
+				}
+
+				HoleIndex++;
+			}
+		}
+
+		AddUndo( MakeUnique<FChangePolygonsVertexInstancesChange>( MoveTemp( RevertInput ) ) );
+	}
+
+	static TArray<FPolygonID> PolygonIDs;
+	PolygonIDs.Reset( VertexInstancesForPolygons.Num() );
+
+	// Perform action
+	{
+		for( const FChangeVertexInstancesForPolygon& VertexInstancesForPolygon : VertexInstancesForPolygons )
+		{
+			const FPolygonID PolygonID = VertexInstancesForPolygon.PolygonID;
+			PolygonIDs.Add( PolygonID );
+			FMeshPolygon& Polygon = Polygons[ PolygonID.GetValue() ];
+
+			for( const FVertexIndexAndInstanceID& IndexAndInstance : VertexInstancesForPolygon.PerimeterVertexIndicesAndInstanceIDs )
+			{
+				// Disconnect old vertex instance from polygon, and connect new one
+				const FVertexInstanceID OldVertexInstanceID = Polygon.PerimeterContour.VertexInstanceIDs[ IndexAndInstance.ContourIndex ];
+				FMeshVertexInstance& OldVertexInstance = VertexInstances[ OldVertexInstanceID.GetValue() ];
+				verify( OldVertexInstance.ConnectedPolygons.Remove( PolygonID ) == 1 );
+
+				FMeshVertexInstance& NewVertexInstance = VertexInstances[ IndexAndInstance.VertexInstanceID.GetValue() ];
+				check( !NewVertexInstance.ConnectedPolygons.Contains( PolygonID ) );
+				NewVertexInstance.ConnectedPolygons.Add( PolygonID );
+
+				Polygon.PerimeterContour.VertexInstanceIDs[ IndexAndInstance.ContourIndex ] = IndexAndInstance.VertexInstanceID;
+
+				// Fix up triangle list
+				for( FMeshTriangle& Triangle : Polygon.Triangles )
+				{
+					for( int32 VertexIndex = 0; VertexIndex < 3; ++VertexIndex )
+					{
+						if( Triangle.VertexInstanceIDs[ VertexIndex ] == OldVertexInstanceID )
+						{
+							Triangle.VertexInstanceIDs[ VertexIndex ] = IndexAndInstance.VertexInstanceID;
+						}
+					}
+				}
+			}
+
+			int32 HoleIndex = 0;
+			for( const FVertexInstancesForPolygonHole& VertexInstancesForPolygonHole : VertexInstancesForPolygon.VertexIndicesAndInstanceIDsForEachHole )
+			{
+				for( const FVertexIndexAndInstanceID& IndexAndInstance : VertexInstancesForPolygonHole.VertexIndicesAndInstanceIDs )
+				{
+					// Disconnect old vertex instance from polygon, and connect new one
+					const FVertexInstanceID OldVertexInstanceID = Polygon.HoleContours[ HoleIndex ].VertexInstanceIDs[ IndexAndInstance.ContourIndex ];
+					FMeshVertexInstance& OldVertexInstance = VertexInstances[ OldVertexInstanceID.GetValue() ];
+					verify( OldVertexInstance.ConnectedPolygons.Remove( PolygonID ) == 1 );
+
+					FMeshVertexInstance& NewVertexInstance = VertexInstances[ IndexAndInstance.VertexInstanceID.GetValue() ];
+					check( !NewVertexInstance.ConnectedPolygons.Contains( PolygonID ) );
+					NewVertexInstance.ConnectedPolygons.Add( PolygonID );
+
+					Polygon.HoleContours[ HoleIndex ].VertexInstanceIDs[ IndexAndInstance.ContourIndex ] = IndexAndInstance.VertexInstanceID;
+
+					// Fix up triangle list
+					for( FMeshTriangle& Triangle : Polygon.Triangles )
+					{
+						for( int32 VertexIndex = 0; VertexIndex < 3; ++VertexIndex )
+						{
+							if( Triangle.VertexInstanceIDs[ VertexIndex ] == OldVertexInstanceID )
+							{
+								Triangle.VertexInstanceIDs[ VertexIndex ] = IndexAndInstance.VertexInstanceID;
+							}
+						}
+					}
+				}
+
+				HoleIndex++;
+			}
+		}
+	}
+
+	// Let the adapter deal with it
+	for( UEditableMeshAdapter* Adapter : Adapters )
+	{
+		Adapter->OnChangePolygonVertexInstances( this, PolygonIDs );
+	}
+
+	UE_LOG( LogEditableMesh, Verbose, TEXT( "* ChangePolygonsVertexInstances returned" ) );
+}
+
+
+FVertexInstanceID UEditableMesh::GetVertexInstanceInPolygonForVertex( const FPolygonID PolygonID, const FVertexID VertexID ) const
+{
+	const FMeshVertex& Vertex = Vertices[ VertexID.GetValue() ];
+
+	for( const FVertexInstanceID VertexInstanceID : Vertex.VertexInstanceIDs )
+	{
+		const FMeshVertexInstance& VertexInstance = VertexInstances[ VertexInstanceID.GetValue() ];
+		if( VertexInstance.ConnectedPolygons.Contains( PolygonID ) )
+		{
+			return VertexInstanceID;
+		}
+	}
+
+	return FVertexInstanceID::Invalid;
+}
+
+
+void UEditableMesh::GetConnectedSoftEdges( const FVertexID VertexID, TArray<FEdgeID>& OutConnectedSoftEdges ) const
+{
+	OutConnectedSoftEdges.Reset();
+
+	const FMeshVertex& Vertex = Vertices[ VertexID.GetValue() ];
+	for( const FEdgeID ConnectedEdgeID : Vertex.ConnectedEdgeIDs )
+	{
+		const FMeshEdge& Edge = Edges[ ConnectedEdgeID.GetValue() ];
+		if( !Edge.bIsHardEdge )
+		{
+			OutConnectedSoftEdges.Add( ConnectedEdgeID );
+		}
+	}
+}
+
+
+void UEditableMesh::GetVertexInstancesInSameSoftEdgedGroup( const FVertexID VertexID, const FPolygonID PolygonID, const bool bPolygonNotYetInitialized, TArray<FVertexInstanceID>& OutVertexInstanceIDs ) const
+{
+	OutVertexInstanceIDs.Reset();
+
+	// Cache a list of all soft edges which share this vertex.
+	// We're not interested in hard edges as they denote a transition to a different soft edged group.
+	// Note that a vertex is usually only treated as 'hard' if it is met by two hard edges: if a hard edge terminates at the vertex we're looking at,
+	// it may be possible to find an adjacency path from one side to the other, and hence does not represent a barrier between soft edged groups.
+	// The exception is if the mesh is not closed and the hard edge leads to the mesh boundary.
+	static TArray<FEdgeID> ConnectedSoftEdges;
+	GetConnectedSoftEdges( VertexID, ConnectedSoftEdges );
+
+	// Iterate through adjacent polygons which contain the given vertex, but don't cross a hard edge.
+	// Maintain a list of polygon IDs to be examined. Adjacents are added to the list if suitable.
+	// Add the start poly here.
+	static TArray<FPolygonID> PolygonsToCheck;
+	PolygonsToCheck.Reset();
+	PolygonsToCheck.Add( PolygonID );
+
+	int32 Index = 0;
+	while( Index < PolygonsToCheck.Num() )
+	{
+		const FPolygonID PolygonToCheck = PolygonsToCheck[ Index ];
+		Index++;
+
+		const FVertexInstanceID VertexInstanceID = GetVertexInstanceInPolygonForVertex( PolygonToCheck, VertexID );
+		if( VertexInstanceID != FVertexInstanceID::Invalid || ( bPolygonNotYetInitialized && Index == 0 ) )
+		{
+			// If we got here, either:
+			//  a) the polygon contains the vertex instance referencing the vertex we're looking at, or
+			//  b) the polygon is the one we passed in, which may not yet have its vertex instances assigned, and we've specified bPolygonNotYetInitialized=true to allow this case
+
+			// If this polygon contains an instance of this vertex, add it to the result.
+			if( VertexInstanceID != FVertexInstanceID::Invalid )
+			{
+				OutVertexInstanceIDs.AddUnique( VertexInstanceID );
+			}
+
+			// Now look at its adjacent polygons. If they are joined by a soft edge which includes the vertex we're interested in, we want to consider them.
+			// We take a shortcut by doing this process in reverse: we already know all the soft edges we are interested in, so check if any of them
+			// have the current polygon as an adjacent.
+			for( const FEdgeID ConnectedSoftEdge : ConnectedSoftEdges )
+			{
+				const FMeshEdge& Edge = Edges[ ConnectedSoftEdge.GetValue() ];
+				if( Edge.ConnectedPolygons.Contains( PolygonToCheck ) )
+				{
+					for( const FPolygonID AdjacentPolygon : Edge.ConnectedPolygons )
+					{
+						// Only add new polygons which haven't yet been added to the list. This prevents circular runs of polygons triggering infinite loops.
+						PolygonsToCheck.AddUnique( AdjacentPolygon );
+					}
+				}
+			}
+		}
+	}
+}
+
+
 void UEditableMesh::SetPolygonsVertexAttributes( const TArray<FVertexAttributesForPolygon>& VertexAttributesForPolygons )
 {
-	FSetPolygonsVertexAttributesChangeInput RevertInput;
+	UE_LOG( LogEditableMesh, Verbose, TEXT( "SetPolygonsVertexAttributes: %s" ), *LogHelpers::ArrayToString( VertexAttributesForPolygons ) );
 
-	RevertInput.VertexAttributesForPolygons.Reserve( VertexAttributesForPolygons.Num() );
 	for( const FVertexAttributesForPolygon& VertexAttributesForPolygon : VertexAttributesForPolygons )
 	{
-		FVertexAttributesForPolygon& RevertPolygonVertexAttributes = *new( RevertInput.VertexAttributesForPolygons ) FVertexAttributesForPolygon();
-		RevertPolygonVertexAttributes.PolygonRef = VertexAttributesForPolygon.PolygonRef;
+		const FPolygonID PolygonID = VertexAttributesForPolygon.PolygonID;
+		FMeshPolygon& Polygon = Polygons[ PolygonID.GetValue() ];
 
-		// Set the polygon's vertex attributes
+		SetPolygonContourVertexAttributes( Polygon.PerimeterContour, PolygonID, INDEX_NONE, VertexAttributesForPolygon.PerimeterVertexAttributeLists );
+
+		for( int32 HoleIndex = 0; HoleIndex < VertexAttributesForPolygon.VertexAttributeListsForEachHole.Num(); ++HoleIndex )
 		{
-			// Polygon perimeter
-			if( VertexAttributesForPolygon.PerimeterVertexAttributeLists.Num() > 0 )
-			{
-				const int32 PerimeterVertexCount = GetPolygonPerimeterVertexCount( VertexAttributesForPolygon.PolygonRef );
-				check( VertexAttributesForPolygon.PerimeterVertexAttributeLists.Num() == PerimeterVertexCount );
-
-				RevertPolygonVertexAttributes.PerimeterVertexAttributeLists.SetNum( VertexAttributesForPolygon.PerimeterVertexAttributeLists.Num(), false );
-
-				for( int32 PerimeterVertexNumber = 0; PerimeterVertexNumber < PerimeterVertexCount; ++PerimeterVertexNumber )
-				{
-					const TArray<FMeshElementAttributeData>& PerimeterVertexAttributeList = VertexAttributesForPolygon.PerimeterVertexAttributeLists[ PerimeterVertexNumber ].Attributes;
-
-					RevertPolygonVertexAttributes.PerimeterVertexAttributeLists[ PerimeterVertexNumber ].Attributes.Reserve( PerimeterVertexAttributeList.Num() );
-					for( const FMeshElementAttributeData& PerimeterVertexAttribute : PerimeterVertexAttributeList )
-					{
-						// Back up the attribute
-						FMeshElementAttributeData& RevertAttribute = *new( RevertPolygonVertexAttributes.PerimeterVertexAttributeLists[ PerimeterVertexNumber ].Attributes ) FMeshElementAttributeData(
-							PerimeterVertexAttribute.AttributeName,
-							PerimeterVertexAttribute.AttributeIndex,
-							GetPolygonPerimeterVertexAttribute( VertexAttributesForPolygon.PolygonRef, PerimeterVertexNumber, PerimeterVertexAttribute.AttributeName, PerimeterVertexAttribute.AttributeIndex ) );
-
-						// Set the new attribute
-						SetPolygonPerimeterVertexAttribute_Internal( VertexAttributesForPolygon.PolygonRef, PerimeterVertexNumber, PerimeterVertexAttribute.AttributeName, PerimeterVertexAttribute.AttributeIndex, PerimeterVertexAttribute.AttributeValue );
-					}
-				}
-			}
-
-			// Polygon holes
-			if( VertexAttributesForPolygon.VertexAttributeListsForEachHole.Num() > 0 )
-			{
-				const int32 HoleCount = GetPolygonHoleCount( VertexAttributesForPolygon.PolygonRef );
-				check( VertexAttributesForPolygon.VertexAttributeListsForEachHole.Num() == HoleCount );
-
-				RevertPolygonVertexAttributes.VertexAttributeListsForEachHole.SetNum( HoleCount, false );
-
-				for( int32 HoleNumber = 0; HoleNumber < HoleCount; ++HoleNumber )
-				{
-					const FVertexAttributesForPolygonHole& AttributesForHoleVertices = VertexAttributesForPolygon.VertexAttributeListsForEachHole[ HoleNumber ];
-					if( AttributesForHoleVertices.VertexAttributeList.Num() > 0 )
-					{
-						const int32 HoleVertexCount = GetPolygonHoleVertexCount( VertexAttributesForPolygon.PolygonRef, HoleNumber );
-						check( AttributesForHoleVertices.VertexAttributeList.Num() == HoleVertexCount );
-
-						RevertPolygonVertexAttributes.VertexAttributeListsForEachHole[ HoleNumber ].VertexAttributeList.SetNum( HoleVertexCount, false );
-
-						for( int32 HoleVertexNumber = 0; HoleVertexNumber < HoleVertexCount; ++HoleVertexNumber )
-						{
-							const TArray<FMeshElementAttributeData>& HoleVertexAttributeList = AttributesForHoleVertices.VertexAttributeList[ HoleVertexNumber ].Attributes;
-
-							RevertPolygonVertexAttributes.VertexAttributeListsForEachHole[ HoleNumber ].VertexAttributeList[ HoleVertexNumber ].Attributes.Reserve( HoleVertexAttributeList.Num() );
-							for( const FMeshElementAttributeData& HoleVertexAttribute : HoleVertexAttributeList )
-							{
-								// Back up the attribute
-								FMeshElementAttributeData& RevertAttribute = *new( RevertPolygonVertexAttributes.VertexAttributeListsForEachHole[ HoleNumber ].VertexAttributeList[ HoleVertexNumber ].Attributes ) FMeshElementAttributeData(
-									HoleVertexAttribute.AttributeName,
-									HoleVertexAttribute.AttributeIndex,
-									GetPolygonHoleVertexAttribute( VertexAttributesForPolygon.PolygonRef, HoleNumber, HoleVertexNumber, HoleVertexAttribute.AttributeName, HoleVertexAttribute.AttributeIndex ) );
-
-								// Set the new attribute
-								SetPolygonHoleVertexAttribute_Internal( VertexAttributesForPolygon.PolygonRef, HoleNumber, HoleVertexNumber, HoleVertexAttribute.AttributeName, HoleVertexAttribute.AttributeIndex, HoleVertexAttribute.AttributeValue );
-							}
-						}
-					}
-				}
-			}
+			SetPolygonContourVertexAttributes( Polygon.HoleContours[ HoleIndex ], PolygonID, HoleIndex, VertexAttributesForPolygon.VertexAttributeListsForEachHole[ HoleIndex ].VertexAttributeList );
 		}
 	}
 
-	AddUndo( MakeUnique<FSetPolygonsVertexAttributesChange>( MoveTemp( RevertInput ) ) );
+	UE_LOG( LogEditableMesh, Verbose, TEXT( "* SetPolygonsVertexAttributes returned" ) );
 }
 
 
-void UEditableMesh::TryToRemovePolygonEdge( const FEdgeID EdgeID, bool& bOutWasEdgeRemoved, FPolygonRef& OutNewPolygonRef )
+bool UEditableMesh::AreAttributeListsEqual( const TArray<FMeshElementAttributeData>& A, const TArray<FMeshElementAttributeData>& B ) const
 {
+	// Check that two attribute lists have attributes of equal values.
+	// It is expected that each list contains exactly the same attribute names and indices, in the same order.
+	check( A.Num() == B.Num() );
+	for( int32 Index = 0; Index < A.Num(); ++Index )
+	{
+		check( A[ Index ].AttributeName == B[ Index ].AttributeName );
+		check( A[ Index ].AttributeIndex == B[ Index ].AttributeIndex );
+		if( A[ Index ].AttributeValue != B[ Index ].AttributeValue )
+		{
+			return false;
+		}
+	}
+	return true;
+};
+
+
+void UEditableMesh::SetPolygonContourVertexAttributes( FMeshPolygonContour& Contour, const FPolygonID PolygonID, const int32 HoleIndex, const TArray<FMeshElementAttributeList>& AttributeLists )
+{
+	const int32 NumContourVertices = Contour.VertexInstanceIDs.Num();
+	check( AttributeLists.Num() == NumContourVertices );
+
+	// Iterate round all polygons in the contour
+	int32 PrevIndex = NumContourVertices - 1;
+	for( int32 Index = 0; Index < NumContourVertices; ++Index )
+	{
+		const FMeshElementAttributeList& AttributeList = AttributeLists[ Index ];
+
+		// If there are no attributes to change, skip this index
+		if( AttributeList.Attributes.Num() == 0 )
+		{
+			continue;
+		}
+
+		// Get vertex instance and vertex.
+		const FVertexInstanceID VertexInstanceID = Contour.VertexInstanceIDs[ Index ];
+		const FMeshVertexInstance& VertexInstance = VertexInstances[ VertexInstanceID.GetValue() ];
+		const FVertexID VertexID = VertexInstance.VertexID;
+
+		// We need to determine whether changing the attributes on this polygon vertex would require vertex instances to be split or merged.
+		// First determine all the existing vertex instances which are in this soft edged group, that is to say,
+		// the group of polygons delimited by hard edges, of which this one is a member.
+		static TArray<FVertexInstanceID> PossibleVertexInstances;
+		PossibleVertexInstances.Reset();
+		const bool bPolygonNotYetInitialized = false;
+		GetVertexInstancesInSameSoftEdgedGroup( VertexID, PolygonID, bPolygonNotYetInitialized, PossibleVertexInstances );
+
+		// Now we have various possibilities:
+		//   1) If any of these vertex instances' attributes match our new polygon vertex attributes, just merge with that one, and disconnect the old one from the polygon.
+		//     1a) If that would cause the old vertex instance to become orphaned (because it is no longer connected to any polygon), it should be deleted.
+		// otherwise:
+		//   2) There are no vertex instances to merge into. 
+		//     2a) If the current vertex instance is not used by any other polygon, change its attributes in-place.
+		//   otherwise:
+		//     2b) Create a new vertex instance with the new attributes, and disconnect the old one from the polygon.
+
+		// Build the complete list of attributes for the polygon vertex being changed
+		static TArray<FMeshElementAttributeData> AllAttributes;
+		GetMergedNamedVertexInstanceAttributeData( GetMergeableVertexInstanceAttributes(), VertexInstanceID, AttributeList.Attributes, AllAttributes );
+
+		// First look for a vertex instance with matching attributes to ours
+		FVertexInstanceID MergedVertexInstanceID = FVertexInstanceID::Invalid;
+
+		for( const FVertexInstanceID PossibleVertexInstance : PossibleVertexInstances )
+		{
+			static TArray<FMeshElementAttributeData> AllAttributesToCompare;
+			GetMergedNamedVertexInstanceAttributeData( GetMergeableVertexInstanceAttributes(), PossibleVertexInstance, TArray<FMeshElementAttributeData>(), AllAttributesToCompare );
+
+			if( AreAttributeListsEqual( AllAttributes, AllAttributesToCompare ) )
+			{
+				// Found a vertex instance we can merge into
+				MergedVertexInstanceID = PossibleVertexInstance;
+				break;
+			}
+		}
+
+		if( MergedVertexInstanceID != FVertexInstanceID::Invalid )
+		{
+			// If we found a vertex instance we can merge into, change the perimeter vertex instance for the new one
+			static TArray<FChangeVertexInstancesForPolygon> VertexInstancesToChange;
+			VertexInstancesToChange.Reset();
+			VertexInstancesToChange.SetNum( 1 );
+
+			VertexInstancesToChange[ 0 ].PolygonID = PolygonID;
+			if( HoleIndex == INDEX_NONE )
+			{
+				VertexInstancesToChange[ 0 ].PerimeterVertexIndicesAndInstanceIDs.Emplace();
+				VertexInstancesToChange[ 0 ].PerimeterVertexIndicesAndInstanceIDs[ 0 ].ContourIndex = Index;
+				VertexInstancesToChange[ 0 ].PerimeterVertexIndicesAndInstanceIDs[ 0 ].VertexInstanceID = MergedVertexInstanceID;
+			}
+			else
+			{
+				VertexInstancesToChange[ 0 ].VertexIndicesAndInstanceIDsForEachHole.SetNum( HoleIndex + 1 );
+				VertexInstancesToChange[ 0 ].VertexIndicesAndInstanceIDsForEachHole[ HoleIndex ].VertexIndicesAndInstanceIDs.Emplace();
+				VertexInstancesToChange[ 0 ].VertexIndicesAndInstanceIDsForEachHole[ HoleIndex ].VertexIndicesAndInstanceIDs[ 0 ].ContourIndex = Index;
+				VertexInstancesToChange[ 0 ].VertexIndicesAndInstanceIDsForEachHole[ HoleIndex ].VertexIndicesAndInstanceIDs[ 0 ].VertexInstanceID = MergedVertexInstanceID;
+			}
+
+			ChangePolygonsVertexInstances( VertexInstancesToChange );
+
+			// If the change orphaned the old one, delete it
+			// @todo mesheditor: make this optional?
+			if( GetVertexInstanceConnectedPolygonCount( VertexInstanceID ) == 0 )
+			{
+				static TArray<FVertexInstanceID> VertexInstanceIDsToDelete;
+				VertexInstanceIDsToDelete.SetNumUninitialized( 1 );
+				VertexInstanceIDsToDelete[ 0 ] = VertexInstanceID;
+
+				const bool bDeleteOrphanedVertices = false;
+				DeleteVertexInstances( VertexInstanceIDsToDelete, bDeleteOrphanedVertices );
+			}
+		}
+		else
+		{
+			if( GetVertexInstanceConnectedPolygonCount( VertexInstanceID ) == 1 )
+			{
+				// This is the only polygon using this vertex instance, so change it in place
+				static TArray<FAttributesForVertexInstance> AttributesForVertexInstance;
+				AttributesForVertexInstance.Reset();
+				AttributesForVertexInstance.SetNum( 1 );
+				AttributesForVertexInstance[ 0 ].VertexInstanceID = VertexInstanceID;
+				AttributesForVertexInstance[ 0 ].VertexInstanceAttributes = AttributeList;
+
+				SetVertexInstancesAttributes( AttributesForVertexInstance );
+			}
+			else
+			{
+				// Split vertex instance: create a new one
+				static TArray<FVertexInstanceToCreate> VertexInstancesToCreate;
+				VertexInstancesToCreate.Reset();
+				VertexInstancesToCreate.SetNum( 1 );
+				VertexInstancesToCreate[ 0 ].VertexID = VertexID;
+				VertexInstancesToCreate[ 0 ].VertexInstanceAttributes.Attributes = AllAttributes;
+
+				static TArray<FVertexInstanceID> NewVertexInstanceIDs;
+				CreateVertexInstances( VertexInstancesToCreate, NewVertexInstanceIDs );
+
+				// and set it on the contour
+				static TArray<FChangeVertexInstancesForPolygon> VertexInstancesToChange;
+				VertexInstancesToChange.Reset();
+				VertexInstancesToChange.SetNum( 1 );
+
+				VertexInstancesToChange[ 0 ].PolygonID = PolygonID;
+				if( HoleIndex == INDEX_NONE )
+				{
+					VertexInstancesToChange[ 0 ].PerimeterVertexIndicesAndInstanceIDs.Emplace();
+					VertexInstancesToChange[ 0 ].PerimeterVertexIndicesAndInstanceIDs[ 0 ].ContourIndex = Index;
+					VertexInstancesToChange[ 0 ].PerimeterVertexIndicesAndInstanceIDs[ 0 ].VertexInstanceID = NewVertexInstanceIDs[ 0 ];
+				}
+				else
+				{
+					VertexInstancesToChange[ 0 ].VertexIndicesAndInstanceIDsForEachHole.SetNum( HoleIndex + 1 );
+					VertexInstancesToChange[ 0 ].VertexIndicesAndInstanceIDsForEachHole[ HoleIndex ].VertexIndicesAndInstanceIDs.Emplace();
+					VertexInstancesToChange[ 0 ].VertexIndicesAndInstanceIDsForEachHole[ HoleIndex ].VertexIndicesAndInstanceIDs[ 0 ].ContourIndex = Index;
+					VertexInstancesToChange[ 0 ].VertexIndicesAndInstanceIDsForEachHole[ HoleIndex ].VertexIndicesAndInstanceIDs[ 0 ].VertexInstanceID = NewVertexInstanceIDs[ 0 ];
+				}
+
+				ChangePolygonsVertexInstances( VertexInstancesToChange );
+			}
+		}
+
+		PrevIndex = Index;
+	}
+}
+
+
+void UEditableMesh::TryToRemovePolygonEdge( const FEdgeID EdgeID, bool& bOutWasEdgeRemoved, FPolygonID& OutNewPolygonID )
+{
+	UE_LOG( LogEditableMesh, Verbose, TEXT( "TryToRemovePolygonEdge: %s" ), *EdgeID.ToString() );
+
 	bOutWasEdgeRemoved = false;
-	OutNewPolygonRef = FPolygonRef::Invalid;
+	OutNewPolygonID = FPolygonID::Invalid;
 
 	// If the edge is not shared by at least two polygons, we can't remove it.  (We would have to delete the polygon that owns
 	// this edge, which is not the intent of this feature.).  We also can't cleanly remove edges that are joining more than two 
@@ -3446,19 +5434,18 @@
 		{
 			TArray<FVertexID> PolygonAVertexIDs, PolygonBVertexIDs;
 
-			const FPolygonRef PolygonARef = GetEdgeConnectedPolygon( EdgeID, 0 );
-			GetPolygonPerimeterVertices( PolygonARef, /* Out */ PolygonAVertexIDs );
+			const FPolygonID PolygonAID = GetEdgeConnectedPolygon( EdgeID, 0 );
+			GetPolygonPerimeterVertices( PolygonAID, /* Out */ PolygonAVertexIDs );
 
-			const FPolygonRef PolygonBRef = GetEdgeConnectedPolygon( EdgeID, 1 );
-			GetPolygonPerimeterVertices( PolygonBRef, /* Out */ PolygonBVertexIDs );
+			const FPolygonID PolygonBID = GetEdgeConnectedPolygon( EdgeID, 1 );
+			GetPolygonPerimeterVertices( PolygonBID, /* Out */ PolygonBVertexIDs );
 
-			// If the polygons are in different sections, we can't remove the edge because we can't determine
-			// which section the replacing polygon should belong to
-			// @todo mesheditor: We could just allow this, and use the first section ID.  It's just a bit weird in the UI.
-			if( PolygonARef.SectionID == PolygonBRef.SectionID )
+			// If the polygons are in different polygon groups, we can't remove the edge because we can't determine
+			// which polygon group the replacing polygon should belong to
+			// @todo mesheditor: We could just allow this, and use the first polygon group ID.  It's just a bit weird in the UI.
+			const FPolygonGroupID PolygonGroupID = GetGroupForPolygon( PolygonAID );
+			if( PolygonGroupID == GetGroupForPolygon( PolygonBID ) )
 			{
-				const FSectionID NewPolygonSectionID = PolygonARef.SectionID;
-
 				// Create a polygon by combining the edges from either polygon we're connected to, omitting the edge we're removing
 				TArray<FVertexAndAttributes> NewPolygonVertices;
 				{
@@ -3527,14 +5514,14 @@
 							NewPolygonVertex.PolygonVertexAttributes.Attributes.Emplace(
 								UEditableMeshAttribute::VertexTextureCoordinate(),
 								TextureCoordinateIndex,
-								GetPolygonPerimeterVertexAttribute( PolygonARef, PolygonAVertexNumber, UEditableMeshAttribute::VertexTextureCoordinate(), TextureCoordinateIndex )
+								GetPolygonPerimeterVertexAttribute( PolygonAID, PolygonAVertexNumber, UEditableMeshAttribute::VertexTextureCoordinate(), TextureCoordinateIndex )
 							);
 						}
 
 						NewPolygonVertex.PolygonVertexAttributes.Attributes.Emplace(
 							UEditableMeshAttribute::VertexColor(),
 							0,
-							GetPolygonPerimeterVertexAttribute( PolygonARef, PolygonAVertexNumber, UEditableMeshAttribute::VertexColor(), 0 )
+							GetPolygonPerimeterVertexAttribute( PolygonAID, PolygonAVertexNumber, UEditableMeshAttribute::VertexColor(), 0 )
 						);
 
 						PolygonAVertexNumber = ( PolygonAVertexNumber + 1 ) % PolygonAVertexIDs.Num();
@@ -3556,14 +5543,14 @@
 							NewPolygonVertex.PolygonVertexAttributes.Attributes.Emplace(
 								UEditableMeshAttribute::VertexTextureCoordinate(),
 								TextureCoordinateIndex,
-								GetPolygonPerimeterVertexAttribute( PolygonBRef, PolygonBVertexNumber, UEditableMeshAttribute::VertexTextureCoordinate(), TextureCoordinateIndex )
+								GetPolygonPerimeterVertexAttribute( PolygonBID, PolygonBVertexNumber, UEditableMeshAttribute::VertexTextureCoordinate(), TextureCoordinateIndex )
 							);
 						}
 
 						NewPolygonVertex.PolygonVertexAttributes.Attributes.Emplace(
 							UEditableMeshAttribute::VertexColor(),
 							0,
-							GetPolygonPerimeterVertexAttribute( PolygonBRef, PolygonBVertexNumber, UEditableMeshAttribute::VertexColor(), 0 )
+							GetPolygonPerimeterVertexAttribute( PolygonBID, PolygonBVertexNumber, UEditableMeshAttribute::VertexColor(), 0 )
 						);
 
 						if( bPolygonsWindInSameDirection )
@@ -3588,10 +5575,12 @@
 
 				// OK, we can go ahead and delete the edge and its connected polygons.  We do NOT want to delete any orphaned
 				// edges or vertices though.  We're going to create a new polygon that connects to those right afterwards.
+				// @todo mesheditor: simplify this action so we just create the new polygon from existing vertex instances, instead of remembering all the attributes.
 				const bool bDeleteOrphanedEdges = false;
 				const bool bDeleteOrphanedVertices = false;
-				const bool bDeleteEmptySections = false;
-				DeleteEdgeAndConnectedPolygons( EdgeID, bDeleteOrphanedEdges, bDeleteOrphanedVertices, bDeleteEmptySections );
+				const bool bDeleteOrphanedVertexInstances = false;
+				const bool bDeleteEmptyPolygonGroups = false;
+				DeleteEdgeAndConnectedPolygons( EdgeID, bDeleteOrphanedEdges, bDeleteOrphanedVertices, bDeleteOrphanedVertexInstances, bDeleteEmptyPolygonGroups );
 
 				// Now create a new polygon to replace the two polygons we deleted
 				{
@@ -3601,35 +5590,32 @@
 
 					// @todo mesheditor perf: Ideally we support creating multiple polygons at once and batching up the work
 					FPolygonToCreate& PolygonToCreate = *new( PolygonsToCreate ) FPolygonToCreate();
-					PolygonToCreate.SectionID = NewPolygonSectionID;
+					PolygonToCreate.PolygonGroupID = PolygonGroupID;
 					PolygonToCreate.PerimeterVertices = NewPolygonVertices;	// @todo mesheditor perf: Extra allocatons/copies: Ideally MoveTemp() here but we can't move a STATIC local!
 					// @todo mesheditor hole: Hole support needed!
 
-					static TArray<FPolygonRef> NewPolygonRefs;
-					NewPolygonRefs.Reset();
+					static TArray<FPolygonID> NewPolygonIDs;
+					NewPolygonIDs.Reset();
 					static TArray<FEdgeID> NewEdgeIDs;
 					NewEdgeIDs.Reset();
-					CreatePolygons( PolygonsToCreate, /* Out */ NewPolygonRefs, /* Out */ NewEdgeIDs );	// @todo mesheditor perf: Extra allocatons/copies: Ideally MoveTemp() here but we can't move a STATIC local!
+					CreatePolygons( PolygonsToCreate, /* Out */ NewPolygonIDs, /* Out */ NewEdgeIDs );	// @todo mesheditor perf: Extra allocatons/copies: Ideally MoveTemp() here but we can't move a STATIC local!
 
-					OutNewPolygonRef = NewPolygonRefs[ 0 ];
+					OutNewPolygonID = NewPolygonIDs[ 0 ];
 				}
 
-				// The polygon's triangulation will be different, plus we lost an edge which may have been hard/soft, thus we need new normals. 
-				// Note that the hardness of the other edges will remain the same.
-				static TArray< FPolygonRef > NewPolygons;
-				NewPolygons.Reset();
-				NewPolygons.Add( OutNewPolygonRef );
-				GenerateNormalsAndTangentsForPolygonsAndAdjacents( NewPolygons );
-
 				bOutWasEdgeRemoved = true;
 			}
 		}
 	}
+
+	UE_LOG( LogEditableMesh, Verbose, TEXT( "* TryToRemovePolygonEdge returned %s" ), *OutNewPolygonID.ToString() );
 }
 
 
 void UEditableMesh::TryToRemoveVertex( const FVertexID VertexID, bool& bOutWasVertexRemoved, FEdgeID& OutNewEdgeID )
 {
+	UE_LOG( LogEditableMesh, Verbose, TEXT( "TryToRemoveVertex: %s" ), *VertexID.ToString() );
+
 	bOutWasVertexRemoved = false;
 	OutNewEdgeID = FEdgeID::Invalid;
 
@@ -3671,23 +5657,18 @@
 		// The new edge will be connected to the same polygons as both of the edges we're replacing.  Because
 		// we only support deleting a vertex shared by two edges, the two edges are guaranteed to be connected
 		// to the same exact polygons
-		static TArray<FPolygonRef> NewEdgeConnectedPolygons;
+		static TArray<FPolygonID> NewEdgeConnectedPolygons;
 		GetVertexConnectedPolygons( VertexID, /* Out */ NewEdgeConnectedPolygons );
 
-		// When undoing this change, we'll retriangulate after everything has been hooked up (before we start destroying edges and vertices below)
+		// Remove the vertex from its connected polygons
 		{
-			const bool bOnlyOnUndo = true;
-			RetriangulatePolygons( NewEdgeConnectedPolygons, bOnlyOnUndo );
-		}
-
-		// Remove the vertex from it's connected polygons
-		{
-			for( const FPolygonRef PolygonRef : NewEdgeConnectedPolygons )
+			for( const FPolygonID PolygonID : NewEdgeConnectedPolygons )
 			{
 				// @todo mesheditor holes
-				const int32 PolygonVertexNumber = FindPolygonPerimeterVertexNumberForVertex( PolygonRef, VertexID );
+				const int32 PolygonVertexNumber = FindPolygonPerimeterVertexNumberForVertex( PolygonID, VertexID );
 				check( PolygonVertexNumber != INDEX_NONE );
-				RemovePolygonPerimeterVertices( PolygonRef, PolygonVertexNumber, 1 );
+				const bool bDeleteOrphanedVertexInstances = false;
+				RemovePolygonPerimeterVertices( PolygonID, PolygonVertexNumber, 1, bDeleteOrphanedVertexInstances );
 			}
 		}
 
@@ -3727,39 +5708,39 @@
 		}
 
 		// Update the normals of the affected polygons
-		GenerateNormalsAndTangentsForPolygonsAndAdjacents( NewEdgeConnectedPolygons );
+		PolygonsPendingNewTangentBasis.Append( NewEdgeConnectedPolygons );
 
 		// Retriangulate all of the affected polygons
+		PolygonsPendingTriangulation.Append( NewEdgeConnectedPolygons );
+
+		// Delete the vertex instances and subsequently orphaned vertex
 		{
-			const bool bOnlyOnUndo = false;
-			RetriangulatePolygons( NewEdgeConnectedPolygons, bOnlyOnUndo );
+			FMeshVertex& Vertex = Vertices[ VertexID.GetValue() ];
+			const bool bDeleteOrphanedVertices = true;
+			DeleteVertexInstances( Vertex.VertexInstanceIDs, bDeleteOrphanedVertices );
 		}
 
-		// Delete the (orphaned) vertex
-		{
-			static TArray<FVertexID> OrphanVertexIDsToDelete;
-			OrphanVertexIDsToDelete.Reset();
-			OrphanVertexIDsToDelete.Add( VertexID );
-			DeleteOrphanVertices( OrphanVertexIDsToDelete );
-		}
-
 		bOutWasVertexRemoved = true;
 		OutNewEdgeID = NewEdgeID;
 	}
+
+	UE_LOG( LogEditableMesh, Verbose, TEXT( "* TryToRemoveVertex returned %s" ), *OutNewEdgeID.ToString() );
 }
 
 
-void UEditableMesh::ExtrudePolygons( const TArray<FPolygonRef>& Polygons, const float ExtrudeDistance, const bool bKeepNeighborsTogether, TArray<FPolygonRef>& OutNewExtrudedFrontPolygons )
+void UEditableMesh::ExtrudePolygons( const TArray<FPolygonID>& PolygonIDs, const float ExtrudeDistance, const bool bKeepNeighborsTogether, TArray<FPolygonID>& OutNewExtrudedFrontPolygons )
 {
+	UE_LOG( LogEditableMesh, Verbose, TEXT( "ExtrudePolygons: %s" ), *LogHelpers::ArrayToString( PolygonIDs ) );
+
 	// @todo mesheditor perf: We can make this much faster by batching up polygons together.  Just be careful about how neighbors are handled.
 	OutNewExtrudedFrontPolygons.Reset();
 
 	// Convert our incoming polygon array to a TSet so we can lookup quickly to see which polygons in the mesh sare members of the set
-	static TSet<FPolygonRef> PolygonsSet;
+	static TSet<FPolygonID> PolygonsSet;
 	PolygonsSet.Reset();
-	PolygonsSet.Append( Polygons );
+	PolygonsSet.Append( PolygonIDs );
 	
-	static TArray<FPolygonRef> AllNewPolygons;
+	static TArray<FPolygonID> AllNewPolygons;
 	AllNewPolygons.Reset();
 
 	static TArray<FAttributesForEdge> AttributesForEdges;
@@ -3777,12 +5758,13 @@
 	EdgeUsageCounts.Reset();
 	static TSet<FVertexID> UniqueVertexIDs;
 	UniqueVertexIDs.Reset();
-	for( int32 PolygonIter = 0; PolygonIter < Polygons.Num(); ++PolygonIter )
+
+	for( int32 PolygonIter = 0; PolygonIter < PolygonIDs.Num(); ++PolygonIter )
 	{
-		const FPolygonRef& PolygonRef = Polygons[ PolygonIter ];
+		const FPolygonID PolygonID = PolygonIDs[ PolygonIter ];
 
 		static TArray<FEdgeID> PolygonPerimeterEdgeIDs;
-		GetPolygonPerimeterEdges( PolygonRef, /* Out */ PolygonPerimeterEdgeIDs );
+		GetPolygonPerimeterEdges( PolygonID, /* Out */ PolygonPerimeterEdgeIDs );
 
 		for( const FEdgeID EdgeID : PolygonPerimeterEdgeIDs )
 		{
@@ -3802,7 +5784,7 @@
 
 
 		static TArray<FVertexID> PolygonPerimeterVertexIDs;
-		this->GetPolygonPerimeterVertices( PolygonRef, /* Out */ PolygonPerimeterVertexIDs );
+		this->GetPolygonPerimeterVertices( PolygonID, /* Out */ PolygonPerimeterVertexIDs );
 
 		for( const FVertexID VertexID : PolygonPerimeterVertexIDs )
 		{
@@ -3829,9 +5811,10 @@
 		// without taking into account all of the connected polygons in our set
 		const bool bIsExtrudingSharedEdges = ( PassIndex == 0 );
 
-		for( int32 PolygonIter = 0; PolygonIter < Polygons.Num(); ++PolygonIter )
+		for( int32 PolygonIter = 0; PolygonIter < PolygonIDs.Num(); ++PolygonIter )
 		{
-			const FPolygonRef& PolygonRef = Polygons[ PolygonIter ];
+			const FPolygonID PolygonID = PolygonIDs[ PolygonIter ];
+			const FPolygonGroupID PolygonGroupID = GetGroupForPolygon( PolygonID );
 
 			if( !bKeepNeighborsTogether )
 			{
@@ -3839,14 +5822,14 @@
 			}
 
 			// Map all of the edge vertices to their new extruded counterpart
-			const int32 PerimeterEdgeCount = GetPolygonPerimeterEdgeCount( PolygonRef );
+			const int32 PerimeterEdgeCount = GetPolygonPerimeterEdgeCount( PolygonID );
 			for( int32 PerimeterEdgeNumber = 0; PerimeterEdgeNumber < PerimeterEdgeCount; ++PerimeterEdgeNumber )
 			{
 				// @todo mesheditor perf: We can change GetPolygonPerimeterEdges() to have a version that returns whether winding is reversed or not, and avoid this call entirely.
 				// @todo mesheditor perf: O(log N^2) iteration here. For every edge, for every edge up to this index.  Need to clean this up. 
 				//		--> Also, there are quite a few places where we are stepping through edges in perimeter-order.  We need to have a nice way to walk that.
 				bool bEdgeWindingIsReversedForPolygon;
-				const FEdgeID EdgeID = this->GetPolygonPerimeterEdge( PolygonRef, PerimeterEdgeNumber, /* Out */ bEdgeWindingIsReversedForPolygon );
+				const FEdgeID EdgeID = this->GetPolygonPerimeterEdge( PolygonID, PerimeterEdgeNumber, /* Out */ bEdgeWindingIsReversedForPolygon );
 
 				const bool bIsSharedEdge = bKeepNeighborsTogether && EdgeUsageCounts[ EdgeID ] > 1;
 				if( bIsSharedEdge == bIsExtrudingSharedEdges )
@@ -3887,7 +5870,6 @@
 							const FVertexID ExtrudedVertexID = ExtrudedVertexIDs[ NextAvailableExtrudedVertexIDNumber++ ];
 
 							ExtrudedEdgeVertexIDPtr = &VertexIDToExtrudedCopy.Add( EdgeVertexID, ExtrudedVertexID );
-
 							// Push the vertex out along the polygon's normal
 							const FVector OriginalVertexPosition = this->GetVertexAttribute( EdgeVertexID, UEditableMeshAttribute::VertexPosition(), 0 );
 
@@ -3898,20 +5880,20 @@
 								// generate an extrude direction that's the average of those polygon normals.
 								FVector ExtrudeDirection = FVector::ZeroVector;
 
-								static TArray<FPolygonRef> ConnectedPolygonRefs;
-								GetVertexConnectedPolygons( EdgeVertexID, /* Out */ ConnectedPolygonRefs );
+								static TArray<FPolygonID> ConnectedPolygonIDs;
+								GetVertexConnectedPolygons( EdgeVertexID, /* Out */ ConnectedPolygonIDs );
 
-								static TArray<FPolygonRef> NeighborPolygonRefs;
-								NeighborPolygonRefs.Reset();
-								for( const FPolygonRef ConnectedPolygonRef : ConnectedPolygonRefs )
+								static TArray<FPolygonID> NeighborPolygonIDs;
+								NeighborPolygonIDs.Reset();
+								for( const FPolygonID ConnectedPolygonID : ConnectedPolygonIDs )
 								{
 									// We only care about polygons that are members of the set of polygons we were asked to extrude
-									if( PolygonsSet.Contains( ConnectedPolygonRef ) )
+									if( PolygonsSet.Contains( ConnectedPolygonID ) )
 									{
-										NeighborPolygonRefs.Add( ConnectedPolygonRef );
+										NeighborPolygonIDs.Add( ConnectedPolygonID );
 
 										// We'll need this polygon's normal to figure out where to put the extruded copy of the polygon
-										const FVector NeighborPolygonNormal = ComputePolygonNormal( ConnectedPolygonRef );
+										const FVector NeighborPolygonNormal = ComputePolygonNormal( ConnectedPolygonID );
 										ExtrudeDirection += NeighborPolygonNormal;
 									}
 								}
@@ -3924,9 +5906,9 @@
 								FVector ClosestIntersectionPointWithExtrudedPlanes;
 								float ClosestIntersectionDistanceSquared = TNumericLimits<float>::Max();
 
-								for( const FPolygonRef NeighborPolygonRef : NeighborPolygonRefs )
+								for( const FPolygonID NeighborPolygonID : NeighborPolygonIDs )
 								{
-									const FPlane NeighborPolygonPlane = ComputePolygonPlane( NeighborPolygonRef );
+									const FPlane NeighborPolygonPlane = ComputePolygonPlane( NeighborPolygonID );
 
 									// Push the plane out
 									const FPlane ExtrudedPlane = [NeighborPolygonPlane, ExtrudeDistance]
@@ -3951,7 +5933,7 @@
 							else
 							{
 								// We'll need this polygon's normal to figure out where to put the extruded copy of the polygon
-								const FVector PolygonNormal = ComputePolygonNormal( PolygonRef );
+								const FVector PolygonNormal = ComputePolygonNormal( PolygonID );
 								ExtrudedVertexPosition = OriginalVertexPosition + ExtrudeDistance * PolygonNormal;
 							}
 
@@ -3984,7 +5966,7 @@
 						NewSidePolygonVertices[ 3 ].VertexID = ExtrudedEdgeVertexIDs[ 1 ];
 						OriginalVertexIDs[ 3 ] = EdgeVertexIDs[ 1 ];
 
-						FPolygonRef NewSidePolygonRef;	// Filled in below
+						FPolygonID NewSidePolygonID;	// Filled in below
 						{
 							static TArray<FPolygonToCreate> PolygonsToCreate;
 							PolygonsToCreate.Reset();
@@ -3992,37 +5974,37 @@
 							// Create the polygon
 							// @todo mesheditor perf: Ideally we support creating multiple polygons at once and batching up the work
 							FPolygonToCreate& PolygonToCreate = *new( PolygonsToCreate ) FPolygonToCreate();
-							PolygonToCreate.SectionID = PolygonRef.SectionID;
+							PolygonToCreate.PolygonGroupID = PolygonGroupID;
 							PolygonToCreate.PerimeterVertices = NewSidePolygonVertices;	// @todo mesheditor perf: Copying static array here, ideally allocations could be avoided
 
 							// NOTE: We never create holes in side polygons
 
-							static TArray<FPolygonRef> NewPolygonRefs;
-							NewPolygonRefs.Reset();
+							static TArray<FPolygonID> NewPolygonIDs;
+							NewPolygonIDs.Reset();
 							static TArray<FEdgeID> NewEdgeIDs;
 							NewEdgeIDs.Reset();
-							CreatePolygons( PolygonsToCreate, /* Out */ NewPolygonRefs, /* Out */ NewEdgeIDs );	// @todo mesheditor perf: Extra allocatons/copies: Ideally MoveTemp() here but we can't move a STATIC local!
+							CreatePolygons( PolygonsToCreate, /* Out */ NewPolygonIDs, /* Out */ NewEdgeIDs );	// @todo mesheditor perf: Extra allocatons/copies: Ideally MoveTemp() here but we can't move a STATIC local!
 
-							NewSidePolygonRef = NewPolygonRefs[ 0 ];
+							NewSidePolygonID = NewPolygonIDs[ 0 ];
 						}
-						AllNewPolygons.Add( NewSidePolygonRef );
+						AllNewPolygons.Add( NewSidePolygonID );
 
 						// Copy polygon UVs from the original polygon's vertices
 						{
 							FVertexAttributesForPolygon& PolygonNewAttributes = *new( VertexAttributesForPolygons ) FVertexAttributesForPolygon();
-							PolygonNewAttributes.PolygonRef = NewSidePolygonRef;
+							PolygonNewAttributes.PolygonID = NewSidePolygonID;
 							PolygonNewAttributes.PerimeterVertexAttributeLists.SetNum( NewSidePolygonVertices.Num(), false );
 
 							for( int32 NewVertexIter = 0; NewVertexIter < NewSidePolygonVertices.Num(); ++NewVertexIter )
 							{
-								const int32 PolygonVertexNumber = this->FindPolygonPerimeterVertexNumberForVertex( PolygonRef, OriginalVertexIDs[ NewVertexIter ] );
+								const int32 PolygonVertexNumber = this->FindPolygonPerimeterVertexNumberForVertex( PolygonID, OriginalVertexIDs[ NewVertexIter ] );
 								check( PolygonVertexNumber != INDEX_NONE );
 
 								TArray<FMeshElementAttributeData>& PerimeterVertexNewAttributes = PolygonNewAttributes.PerimeterVertexAttributeLists[ NewVertexIter ].Attributes;
 								for( int32 TextureCoordinateIndex = 0; TextureCoordinateIndex < TextureCoordinateCount; ++TextureCoordinateIndex )
 								{
 									const FVector4 TextureCoordinate = this->GetPolygonPerimeterVertexAttribute(
-										PolygonRef,
+										PolygonID,
 										PolygonVertexNumber,
 										UEditableMeshAttribute::VertexTextureCoordinate(),
 										TextureCoordinateIndex );
@@ -4031,7 +6013,7 @@
 								}
 
 								const FVector4 VertexColor = this->GetPolygonPerimeterVertexAttribute(
-									PolygonRef,
+									PolygonID,
 									PolygonVertexNumber,
 									UEditableMeshAttribute::VertexColor(),
 									0 );
@@ -4042,11 +6024,11 @@
 
 						// All of this edges of the new polygon will be hard
 						{
-							const int32 NewPerimeterEdgeCount = this->GetPolygonPerimeterEdgeCount( NewSidePolygonRef );
+							const int32 NewPerimeterEdgeCount = this->GetPolygonPerimeterEdgeCount( NewSidePolygonID );
 							for( int32 NewPerimeterEdgeNumber = 0; NewPerimeterEdgeNumber < NewPerimeterEdgeCount; ++NewPerimeterEdgeNumber )
 							{
 								bool bNewEdgeWindingIsReversedForPolygon;
-								const FEdgeID NewEdgeID = this->GetPolygonPerimeterEdge( NewSidePolygonRef, NewPerimeterEdgeNumber, /* Out */ bNewEdgeWindingIsReversedForPolygon );
+								const FEdgeID NewEdgeID = this->GetPolygonPerimeterEdge( NewSidePolygonID, NewPerimeterEdgeNumber, /* Out */ bNewEdgeWindingIsReversedForPolygon );
 
 								// New side polygons get hard edges, but not hard crease weights
 								FAttributesForEdge& AttributesForEdge = *new( AttributesForEdges ) FAttributesForEdge();
@@ -4060,15 +6042,16 @@
 		}
 	}
 
-	for( int32 PolygonIter = 0; PolygonIter < Polygons.Num(); ++PolygonIter )
+	for( int32 PolygonIter = 0; PolygonIter < PolygonIDs.Num(); ++PolygonIter )
 	{
-		const FPolygonRef& PolygonRef = Polygons[ PolygonIter ];
+		const FPolygonID PolygonID = PolygonIDs[ PolygonIter ];
+		const FPolygonGroupID PolygonGroupID = GetGroupForPolygon( PolygonID );
 
 		static TArray<FVertexID> PolygonVertexIDs;
-		this->GetPolygonPerimeterVertices( PolygonRef, /* Out */ PolygonVertexIDs );
+		this->GetPolygonPerimeterVertices( PolygonID, /* Out */ PolygonVertexIDs );
 
 		// Create a new extruded polygon for the face
-		FPolygonRef ExtrudedFrontPolygonRef;	// Filled in below
+		FPolygonID ExtrudedFrontPolygonID;	// Filled in below
 		{
 			static TArray<FVertexAndAttributes> NewFrontPolygonVertices;
 			NewFrontPolygonVertices.Reset( PolygonVertexIDs.Num() );
@@ -4097,23 +6080,24 @@
 				// Create the polygon
 				// @todo mesheditor perf: Ideally we support creating multiple polygons at once and batching up the work
 				FPolygonToCreate& PolygonToCreate = *new( PolygonsToCreate ) FPolygonToCreate();
-				PolygonToCreate.SectionID = PolygonRef.SectionID;
+				PolygonToCreate.PolygonGroupID = PolygonGroupID;
+				PolygonToCreate.PolygonEdgeHardness = EPolygonEdgeHardness::AllEdgesHard;
 				PolygonToCreate.PerimeterVertices = NewFrontPolygonVertices;	// @todo mesheditor perf: Copying static array here, ideally allocations could be avoided
 																				// @todo mesheditor holes: This algorithm is ignoring polygon holes!  Should be easy to support this by doing roughly the same thing.
-				static TArray<FPolygonRef> NewPolygonRefs;
-				NewPolygonRefs.Reset();
+				static TArray<FPolygonID> NewPolygonIDs;
+				NewPolygonIDs.Reset();
 				static TArray<FEdgeID> NewEdgeIDs;
 				NewEdgeIDs.Reset();
-				CreatePolygons( PolygonsToCreate, /* Out */ NewPolygonRefs, /* Out */ NewEdgeIDs );	// @todo mesheditor perf: Extra allocatons/copies: Ideally MoveTemp() here but we can't move a STATIC local!
+				CreatePolygons( PolygonsToCreate, /* Out */ NewPolygonIDs, /* Out */ NewEdgeIDs );	// @todo mesheditor perf: Extra allocatons/copies: Ideally MoveTemp() here but we can't move a STATIC local!
 
-				ExtrudedFrontPolygonRef = NewPolygonRefs[ 0 ];
+				ExtrudedFrontPolygonID = NewPolygonIDs[ 0 ];
 			}
-			AllNewPolygons.Add( ExtrudedFrontPolygonRef );
+			AllNewPolygons.Add( ExtrudedFrontPolygonID );
 
 			// Retain the UVs from the original polygon
 			{
 				FVertexAttributesForPolygon& PolygonNewAttributes = *new( VertexAttributesForPolygons ) FVertexAttributesForPolygon();
-				PolygonNewAttributes.PolygonRef = ExtrudedFrontPolygonRef;
+				PolygonNewAttributes.PolygonID = ExtrudedFrontPolygonID;
 				PolygonNewAttributes.PerimeterVertexAttributeLists.SetNum( PolygonVertexIDs.Num(), false );
 
 				for( int32 PolygonVertexNumber = 0; PolygonVertexNumber < PolygonVertexIDs.Num(); ++PolygonVertexNumber )
@@ -4122,7 +6106,7 @@
 					for( int32 TextureCoordinateIndex = 0; TextureCoordinateIndex < TextureCoordinateCount; ++TextureCoordinateIndex )
 					{
 						const FVector4 TextureCoordinate = this->GetPolygonPerimeterVertexAttribute(
-							PolygonRef,
+							PolygonID,
 							PolygonVertexNumber,
 							UEditableMeshAttribute::VertexTextureCoordinate(),
 							TextureCoordinateIndex );
@@ -4131,7 +6115,7 @@
 					}
 
 					const FVector4 VertexColor = this->GetPolygonPerimeterVertexAttribute(
-						PolygonRef,
+						PolygonID,
 						PolygonVertexNumber,
 						UEditableMeshAttribute::VertexColor(),
 						0 );
@@ -4144,16 +6128,16 @@
 			// All of the border edges of the new polygon will be hard.  If it was a shared edge, then we'll just preserve whatever was
 			// originally going on with the internal edge.
 			{
-				const int32 NewPerimeterEdgeCount = this->GetPolygonPerimeterEdgeCount( ExtrudedFrontPolygonRef );
-				check( NewPerimeterEdgeCount == GetPolygonPerimeterEdgeCount( PolygonRef ) );	// New polygon should always have the same number of edges (in the same order) as the original!
+				const int32 NewPerimeterEdgeCount = this->GetPolygonPerimeterEdgeCount( ExtrudedFrontPolygonID );
+				check( NewPerimeterEdgeCount == GetPolygonPerimeterEdgeCount( PolygonID ) );	// New polygon should always have the same number of edges (in the same order) as the original!
 				for( int32 PerimeterEdgeNumber = 0; PerimeterEdgeNumber < NewPerimeterEdgeCount; ++PerimeterEdgeNumber )
 				{
 					bool bOriginalEdgeWindingIsReversedForPolygon;
-					const FEdgeID OriginalEdgeID = this->GetPolygonPerimeterEdge( PolygonRef, PerimeterEdgeNumber, /* Out */ bOriginalEdgeWindingIsReversedForPolygon );
+					const FEdgeID OriginalEdgeID = this->GetPolygonPerimeterEdge( PolygonID, PerimeterEdgeNumber, /* Out */ bOriginalEdgeWindingIsReversedForPolygon );
 					const bool bIsSharedEdge = bKeepNeighborsTogether && EdgeUsageCounts[ OriginalEdgeID ] > 1;
 
 					bool bEdgeWindingIsReversedForPolygon;
-					const FEdgeID EdgeID = this->GetPolygonPerimeterEdge( ExtrudedFrontPolygonRef, PerimeterEdgeNumber, /* Out */ bEdgeWindingIsReversedForPolygon );
+					const FEdgeID EdgeID = this->GetPolygonPerimeterEdge( ExtrudedFrontPolygonID, PerimeterEdgeNumber, /* Out */ bEdgeWindingIsReversedForPolygon );
 
 					const FVector4 NewEdgeHardnessAttribute = bIsSharedEdge ? GetEdgeAttribute( OriginalEdgeID, UEditableMeshAttribute::EdgeIsHard(), 0 ) : FVector( 1.0f );
 
@@ -4164,7 +6148,7 @@
 			}
 		}
 
-		OutNewExtrudedFrontPolygons.Add( ExtrudedFrontPolygonRef );
+		OutNewExtrudedFrontPolygons.Add( ExtrudedFrontPolygonID );
 	}
 	check( NextAvailableExtrudedVertexIDNumber == ExtrudedVertexIDs.Num() );	// Make sure all of the vertices we created were actually used by new polygons
 
@@ -4179,12 +6163,12 @@
 	{
 		const bool bDeleteOrphanedEdges = true;
 		const bool bDeleteOrphanedVertices = true;
-		const bool bDeleteEmptySections = false;
-		DeletePolygons( Polygons, bDeleteOrphanedEdges, bDeleteOrphanedVertices, bDeleteEmptySections );
+		const bool bDeleteOrphanedVertexInstances = true;
+		const bool bDeleteEmptyPolygonGroups = false;
+		DeletePolygons( PolygonIDs, bDeleteOrphanedEdges, bDeleteOrphanedVertices, bDeleteOrphanedVertexInstances, bDeleteEmptyPolygonGroups );
 	}
 
-	// Generate normals for everything, including adjacents because we might have changed edge hardness
-	GenerateNormalsAndTangentsForPolygonsAndAdjacents( AllNewPolygons );
+	UE_LOG( LogEditableMesh, Verbose, TEXT( "* ExtrudePolygons returned %s" ), *LogHelpers::ArrayToString( OutNewExtrudedFrontPolygons ) );
 }
 
 
@@ -4303,21 +6287,21 @@
 
 			FPolygonToCreate& PolygonToCreate = *new( PolygonsToCreate ) FPolygonToCreate();
 
-			// We need to figure out which mesh section to put the new polygons in.  To do this, we'll look at
-			// which polygons are already connected to the current edge, and use section number from the first
-			// polygon we can find.  If no polygons are connected, then we'll just use first section in the mesh.
+			// We need to figure out which mesh PolygonGroup to put the new polygons in.  To do this, we'll look at
+			// which polygons are already connected to the current edge, and use PolygonGroup number from the first
+			// polygon we can find.  If no polygons are connected, then we'll just use first PolygonGroup in the mesh.
 			// We'll also capture texture coordinates from this polygon, so we can apply them to the new polygon vertices.
-			TOptional<FPolygonRef> ConnectedPolygonRef;
+			FPolygonID ConnectedPolygonID = FPolygonID::Invalid;
 			{
 				const int32 ConnectedPolygonCount = GetEdgeConnectedPolygonCount( OriginalEdgeID );
 				if( ConnectedPolygonCount > 0 )
 				{
-					ConnectedPolygonRef = GetEdgeConnectedPolygon( OriginalEdgeID, 0 );
+					ConnectedPolygonID = GetEdgeConnectedPolygon( OriginalEdgeID, 0 );
 				}
 			}
 
-			PolygonToCreate.SectionID = ConnectedPolygonRef.IsSet() ? ConnectedPolygonRef->SectionID : GetFirstValidSection();
-			check( PolygonToCreate.SectionID != FSectionID::Invalid );
+			PolygonToCreate.PolygonGroupID = ( ConnectedPolygonID != FPolygonID::Invalid ) ? GetGroupForPolygon( ConnectedPolygonID ) : GetFirstValidPolygonGroup();
+			check( PolygonToCreate.PolygonGroupID != FPolygonGroupID::Invalid );
 
 			PolygonToCreate.PerimeterVertices.SetNum( 4, false );
 			FVertexID ConnectedPolygonVertexIDsForTextureCoordinates[ 4 ];
@@ -4338,7 +6322,7 @@
 			}
 
 			// Preserve polygon vertex attributes
-			if( ConnectedPolygonRef.IsSet() )
+			if( ConnectedPolygonID != FPolygonID::Invalid )
 			{
 				for( int32 PerimeterVertexNumber = 0; PerimeterVertexNumber < PolygonToCreate.PerimeterVertices.Num(); ++PerimeterVertexNumber )
 				{
@@ -4346,16 +6330,16 @@
 
 					TArray<FMeshElementAttributeData>& AttributesForVertex = PolygonToCreate.PerimeterVertices[ PerimeterVertexNumber ].PolygonVertexAttributes.Attributes;
 				
-					const int32 ConnectedPolygonPerimeterVertexNumber = FindPolygonPerimeterVertexNumberForVertex( ConnectedPolygonRef.GetValue(), ConnectedPolygonVertexIDForTextureCoordinates );
+					const int32 ConnectedPolygonPerimeterVertexNumber = FindPolygonPerimeterVertexNumberForVertex( ConnectedPolygonID, ConnectedPolygonVertexIDForTextureCoordinates );
 					check( ConnectedPolygonPerimeterVertexNumber != INDEX_NONE );
 
 					for( int32 TextureCoordinateIndex = 0; TextureCoordinateIndex < TextureCoordinateCount; ++TextureCoordinateIndex )
 					{
-						const FVector4 TextureCoordinate = GetPolygonPerimeterVertexAttribute( ConnectedPolygonRef.GetValue(), ConnectedPolygonPerimeterVertexNumber, UEditableMeshAttribute::VertexTextureCoordinate(), TextureCoordinateIndex );
+						const FVector4 TextureCoordinate = GetPolygonPerimeterVertexAttribute( ConnectedPolygonID, ConnectedPolygonPerimeterVertexNumber, UEditableMeshAttribute::VertexTextureCoordinate(), TextureCoordinateIndex );
 						AttributesForVertex.Add( FMeshElementAttributeData( UEditableMeshAttribute::VertexTextureCoordinate(), TextureCoordinateIndex, TextureCoordinate ) );
 					}
 
-					const FVector4 VertexColor = GetPolygonPerimeterVertexAttribute( ConnectedPolygonRef.GetValue(), ConnectedPolygonPerimeterVertexNumber, UEditableMeshAttribute::VertexColor(), 0 );
+					const FVector4 VertexColor = GetPolygonPerimeterVertexAttribute( ConnectedPolygonID, ConnectedPolygonPerimeterVertexNumber, UEditableMeshAttribute::VertexColor(), 0 );
 					AttributesForVertex.Add( FMeshElementAttributeData( UEditableMeshAttribute::VertexColor(), 0, VertexColor ) );
 				}
 			}
@@ -4363,18 +6347,15 @@
 
 		// Create the polygons.  Note that this will also automatically create the missing side edges that connect
 		// the original edge to it's extended counterpart.
-		static TArray<FPolygonRef> NewPolygonRefs;
-		NewPolygonRefs.Reset();
+		static TArray<FPolygonID> NewPolygonIDs;
+		NewPolygonIDs.Reset();
 		static TArray<FEdgeID> NewEdgeIDs;
 		NewEdgeIDs.Reset();
-		CreatePolygons( PolygonsToCreate, /* Out */ NewPolygonRefs, /* Out */ NewEdgeIDs );
+		CreatePolygons( PolygonsToCreate, /* Out */ NewPolygonIDs, /* Out */ NewEdgeIDs );
 
 		// Expecting no more than two new edges to be created while creating the polygon.  It's possible for zero or one edge
 		// to be created, depending on how many edges we share with neighbors that were extended.
 		check( bWeldNeighbors ? ( NewEdgeIDs.Num() <= 2 * EdgeIDs.Num() ) : ( NewEdgeIDs.Num() == 2 * EdgeIDs.Num() ) );
-
-		// Generate normals for all of the new polygons
-		GenerateNormalsAndTangentsForPolygons( NewPolygonRefs );
 	}
 }
 
@@ -4467,33 +6448,33 @@
 					FPolygonToCreate& PolygonToCreate = *new( PolygonsToCreate ) FPolygonToCreate();
 
 					// Figure out which of the connected polygons shares the edge we're going to be using
-					TOptional<FPolygonRef> ConnectedPolygonRef;
+					FPolygonID ConnectedPolygonID = FPolygonID::Invalid;
 					{
-						static TArray<FPolygonRef> ConnectedPolygonRefs;
-						GetVertexConnectedPolygons( OriginalVertexID, /* Out */ ConnectedPolygonRefs );
+						static TArray<FPolygonID> ConnectedPolygonIDs;
+						GetVertexConnectedPolygons( OriginalVertexID, /* Out */ ConnectedPolygonIDs );
 
-						for( const FPolygonRef PolygonRef : ConnectedPolygonRefs )
+						for( const FPolygonID PolygonID : ConnectedPolygonIDs )
 						{
-							const int32 AdjacentVertexNumber = FindPolygonPerimeterVertexNumberForVertex( PolygonRef, AdjacentVertexID );
+							const int32 AdjacentVertexNumber = FindPolygonPerimeterVertexNumberForVertex( PolygonID, AdjacentVertexID );
 							if( AdjacentVertexNumber != INDEX_NONE )
 							{
 								// NOTE: There can be more than one polygon that shares this edge.  We'll just take the first one we find.
-								ConnectedPolygonRef = PolygonRef;
+								ConnectedPolygonID = PolygonID;
 								break;
 							}
 						}
 					}
 
 					bool bConnectedPolygonWindsForward = true;
-					if( ConnectedPolygonRef.IsSet() )
+					if( ConnectedPolygonID != FPolygonID::Invalid )
 					{
-						const int32 OriginalVertexNumberOnConnectedPolygon = FindPolygonPerimeterVertexNumberForVertex( ConnectedPolygonRef.GetValue(), OriginalVertexID );
+						const int32 OriginalVertexNumberOnConnectedPolygon = FindPolygonPerimeterVertexNumberForVertex( ConnectedPolygonID, OriginalVertexID );
 						check( OriginalVertexNumberOnConnectedPolygon != INDEX_NONE );
 
-						const int32 AdjacentVertexNumberOnConnectedPolygon = FindPolygonPerimeterVertexNumberForVertex( ConnectedPolygonRef.GetValue(), AdjacentVertexID );
+						const int32 AdjacentVertexNumberOnConnectedPolygon = FindPolygonPerimeterVertexNumberForVertex( ConnectedPolygonID, AdjacentVertexID );
 						check( AdjacentVertexNumberOnConnectedPolygon != INDEX_NONE );
 
-						const int32 PerimeterVertexCount = GetPolygonPerimeterVertexCount( ConnectedPolygonRef.GetValue() );
+						const int32 PerimeterVertexCount = GetPolygonPerimeterVertexCount( ConnectedPolygonID );
 						if( !( OriginalVertexNumberOnConnectedPolygon == ( PerimeterVertexCount - 1 ) && AdjacentVertexNumberOnConnectedPolygon == 0 ) &&  // Doesn't wrap forward?
 							( OriginalVertexNumberOnConnectedPolygon > AdjacentVertexNumberOnConnectedPolygon ||			// Winds backwards?
 							  AdjacentVertexNumberOnConnectedPolygon == ( PerimeterVertexCount - 1 ) && OriginalVertexNumberOnConnectedPolygon == 0 ) )	// Wraps backwards?
@@ -4502,8 +6483,8 @@
 						}
 					}
 
-					PolygonToCreate.SectionID = ConnectedPolygonRef.IsSet() ? ConnectedPolygonRef->SectionID : GetFirstValidSection();
-					check( PolygonToCreate.SectionID != FSectionID::Invalid );
+					PolygonToCreate.PolygonGroupID = ( ConnectedPolygonID != FPolygonID::Invalid ) ? GetGroupForPolygon( ConnectedPolygonID ) : GetFirstValidPolygonGroup();
+					check( PolygonToCreate.PolygonGroupID != FPolygonGroupID::Invalid );
 
 					FVertexID ConnectedPolygonVertexIDsForTextureCoordinates[ 3 ];
 					PolygonToCreate.PerimeterVertices.SetNum( 3, false );
@@ -4537,7 +6518,7 @@
 					}
 
 					// Preserve polygon vertex attributes
-					if( ConnectedPolygonRef.IsSet() )
+					if( ConnectedPolygonID != FPolygonID::Invalid )
 					{
 						for( int32 PerimeterVertexNumber = 0; PerimeterVertexNumber < PolygonToCreate.PerimeterVertices.Num(); ++PerimeterVertexNumber )
 						{
@@ -4545,16 +6526,16 @@
 
 							TArray<FMeshElementAttributeData>& AttributesForVertex = PolygonToCreate.PerimeterVertices[ PerimeterVertexNumber ].PolygonVertexAttributes.Attributes;
 
-							const int32 ConnectedPolygonPerimeterVertexNumber = FindPolygonPerimeterVertexNumberForVertex( ConnectedPolygonRef.GetValue(), ConnectedPolygonVertexIDForTextureCoordinates );
+							const int32 ConnectedPolygonPerimeterVertexNumber = FindPolygonPerimeterVertexNumberForVertex( ConnectedPolygonID, ConnectedPolygonVertexIDForTextureCoordinates );
 							check( ConnectedPolygonPerimeterVertexNumber != INDEX_NONE );
 
 							for( int32 TextureCoordinateIndex = 0; TextureCoordinateIndex < TextureCoordinateCount; ++TextureCoordinateIndex )
 							{
-								const FVector4 TextureCoordinate = GetPolygonPerimeterVertexAttribute( ConnectedPolygonRef.GetValue(), ConnectedPolygonPerimeterVertexNumber, UEditableMeshAttribute::VertexTextureCoordinate(), TextureCoordinateIndex );
+								const FVector4 TextureCoordinate = GetPolygonPerimeterVertexAttribute( ConnectedPolygonID, ConnectedPolygonPerimeterVertexNumber, UEditableMeshAttribute::VertexTextureCoordinate(), TextureCoordinateIndex );
 								AttributesForVertex.Add( FMeshElementAttributeData( UEditableMeshAttribute::VertexTextureCoordinate(), TextureCoordinateIndex, TextureCoordinate ) );
 							}
 
-							const FVector4 VertexColor = GetPolygonPerimeterVertexAttribute( ConnectedPolygonRef.GetValue(), ConnectedPolygonPerimeterVertexNumber, UEditableMeshAttribute::VertexColor(), 0 );
+							const FVector4 VertexColor = GetPolygonPerimeterVertexAttribute( ConnectedPolygonID, ConnectedPolygonPerimeterVertexNumber, UEditableMeshAttribute::VertexColor(), 0 );
 							AttributesForVertex.Add( FMeshElementAttributeData( UEditableMeshAttribute::VertexColor(), 0, VertexColor ) );
 						}
 					}
@@ -4562,31 +6543,28 @@
 			}		
 		}
 
-		static TArray<FPolygonRef> NewPolygonRefs;
-		NewPolygonRefs.Reset();
+		static TArray<FPolygonID> NewPolygonIDs;
+		NewPolygonIDs.Reset();
 
 		// Create the polygons.  Note that this will also automatically create the missing side edges that connect
 		// the original edge to it's extended counterpart.
 		static TArray<FEdgeID> NewEdgeIDs;
 		NewEdgeIDs.Reset();
-		CreatePolygons( PolygonsToCreate, /* Out */ NewPolygonRefs, /* Out */ NewEdgeIDs );
-
-		// Generate normals for all of the new polygons
-		GenerateNormalsAndTangentsForPolygons( NewPolygonRefs );
+		CreatePolygons( PolygonsToCreate, /* Out */ NewPolygonIDs, /* Out */ NewEdgeIDs );
 	}
 }
 
 
-void UEditableMesh::ComputePolygonsSharedEdges( const TArray<FPolygonRef>& PolygonRefs, TArray<FEdgeID>& OutSharedEdgeIDs ) const
+void UEditableMesh::ComputePolygonsSharedEdges( const TArray<FPolygonID>& PolygonIDs, TArray<FEdgeID>& OutSharedEdgeIDs ) const
 {
 	OutSharedEdgeIDs.Reset();
 
 	static TSet<FEdgeID> EdgesSeenSoFar;
 	EdgesSeenSoFar.Reset();
-	for( const FPolygonRef PolygonRef : PolygonRefs )
+	for( const FPolygonID PolygonID : PolygonIDs )
 	{
 		static TArray<FEdgeID> PolygonPerimeterEdgeIDs;
-		GetPolygonPerimeterEdges( PolygonRef, /* Out */ PolygonPerimeterEdgeIDs );
+		GetPolygonPerimeterEdges( PolygonID, /* Out */ PolygonPerimeterEdgeIDs );
 
 		for( const FEdgeID EdgeID : PolygonPerimeterEdgeIDs )
 		{
@@ -4603,7 +6581,7 @@
 }
 
 
-void UEditableMesh::BevelOrInsetPolygons( const TArray<FPolygonRef>& PolygonRefs, const float InsetFixedDistance, const float InsetProgressTowardCenter, const EInsetPolygonsMode Mode, const bool bShouldBevel, TArray<FPolygonRef>& OutNewCenterPolygonRefs, TArray<FPolygonRef>& OutNewSidePolygonRefs )
+void UEditableMesh::BevelOrInsetPolygons( const TArray<FPolygonID>& PolygonIDs, const float InsetFixedDistance, const float InsetProgressTowardCenter, const EInsetPolygonsMode Mode, const bool bShouldBevel, TArray<FPolygonID>& OutNewCenterPolygonIDs, TArray<FPolygonID>& OutNewSidePolygonIDs )
 {
 	// @todo mesheditor inset/bevel: Weird feedback loop issues at glancing angles when moving mouse while beveling
 
@@ -4616,13 +6594,15 @@
 	static TArray<FAttributesForVertex> AttributesForVertices;
 	AttributesForVertices.Reset();
 
-	for( const FPolygonRef PolygonRef : PolygonRefs )
+	for( const FPolygonID PolygonID : PolygonIDs )
 	{
+		const FPolygonGroupID PolygonGroupID = GetGroupForPolygon( PolygonID );
+
 		// Find the center of this polygon
-		const FVector PolygonCenter = ComputePolygonCenter( PolygonRef );
+		const FVector PolygonCenter = ComputePolygonCenter( PolygonID );
 
 		static TArray<FVertexID> PerimeterVertexIDs;
-		GetPolygonPerimeterVertices( PolygonRef, /* Out */ PerimeterVertexIDs );
+		GetPolygonPerimeterVertices( PolygonID, /* Out */ PerimeterVertexIDs );
 
 		static TArray<FVertexToCreate> VerticesToCreate;
 		VerticesToCreate.Reset( PerimeterVertexIDs.Num() );
@@ -4653,14 +6633,14 @@
 				int32 ConnectedBevelPolygonCount = 0;
 				bool bIsOnlyConnectedToBevelPolygons = true;
 
-				static TArray<FPolygonRef> ConnectedPolygonRefs;
-				GetVertexConnectedPolygons( PerimeterVertexID, /* Out */ ConnectedPolygonRefs );
-				for( const FPolygonRef ConnectedPolygonRef : ConnectedPolygonRefs )
+				static TArray<FPolygonID> ConnectedPolygonIDs;
+				GetVertexConnectedPolygons( PerimeterVertexID, /* Out */ ConnectedPolygonIDs );
+				for( const FPolygonID ConnectedPolygonID : ConnectedPolygonIDs )
 				{
-					if( PolygonRefs.Contains( ConnectedPolygonRef ) )
+					if( PolygonIDs.Contains( ConnectedPolygonID ) )
 					{
 						++ConnectedBevelPolygonCount;
-						const FVector ConnectedPolygonNormal = ComputePolygonNormal( ConnectedPolygonRef );
+						const FVector ConnectedPolygonNormal = ComputePolygonNormal( ConnectedPolygonID );
 						OffsetDirection += -ConnectedPolygonNormal;
 					}
 					else
@@ -4686,10 +6666,10 @@
 			VertexToCreate.VertexAttributes.Attributes.Add( FMeshElementAttributeData( UEditableMeshAttribute::VertexPosition(), 0, InsetVertexPosition ) );
 
 			// Interpolate the texture coordinates and vertex colors on this polygon to figure out UVs for our new vertex
-			static TArray<int32> PerimeterVertexIndices;
+			FMeshTriangle Triangle;
 			FVector TriangleVertexWeights;
 
-			if( ComputeBarycentricWeightForPointOnPolygon( PolygonRef, InsetVertexPosition, PerimeterVertexIndices, TriangleVertexWeights ) )
+			if( ComputeBarycentricWeightForPointOnPolygon( PolygonID, InsetVertexPosition, Triangle, TriangleVertexWeights ) )
 			{
 				TArray<FVector4>& TextureCoordinatesForNewVertex = *new( TextureCoordinatesForNewVertices ) TArray<FVector4>();
 				TextureCoordinatesForNewVertex.SetNum( TextureCoordinateCount, false );
@@ -4697,15 +6677,15 @@
 				for( int32 TextureCoordinateIndex = 0; TextureCoordinateIndex < TextureCoordinateCount; ++TextureCoordinateIndex )
 				{
 					TextureCoordinatesForNewVertex[ TextureCoordinateIndex ] =
-						TriangleVertexWeights.X * GetPolygonPerimeterVertexAttribute( PolygonRef, PerimeterVertexIndices[ 0 ], UEditableMeshAttribute::VertexTextureCoordinate(), TextureCoordinateIndex ) +
-						TriangleVertexWeights.Y * GetPolygonPerimeterVertexAttribute( PolygonRef, PerimeterVertexIndices[ 1 ], UEditableMeshAttribute::VertexTextureCoordinate(), TextureCoordinateIndex ) +
-						TriangleVertexWeights.Z * GetPolygonPerimeterVertexAttribute( PolygonRef, PerimeterVertexIndices[ 2 ], UEditableMeshAttribute::VertexTextureCoordinate(), TextureCoordinateIndex );
+						TriangleVertexWeights.X * GetVertexInstanceAttribute( Triangle.VertexInstanceIDs[ 0 ], UEditableMeshAttribute::VertexTextureCoordinate(), TextureCoordinateIndex ) +
+						TriangleVertexWeights.Y * GetVertexInstanceAttribute( Triangle.VertexInstanceIDs[ 1 ], UEditableMeshAttribute::VertexTextureCoordinate(), TextureCoordinateIndex ) +
+						TriangleVertexWeights.Z * GetVertexInstanceAttribute( Triangle.VertexInstanceIDs[ 2 ], UEditableMeshAttribute::VertexTextureCoordinate(), TextureCoordinateIndex );
 				}
 
 				VertexColorsForNewVertices.Add(
-					TriangleVertexWeights.X * GetPolygonPerimeterVertexAttribute( PolygonRef, PerimeterVertexIndices[ 0 ], UEditableMeshAttribute::VertexColor(), 0 ) +
-					TriangleVertexWeights.Y * GetPolygonPerimeterVertexAttribute( PolygonRef, PerimeterVertexIndices[ 1 ], UEditableMeshAttribute::VertexColor(), 0 ) +
-					TriangleVertexWeights.Z * GetPolygonPerimeterVertexAttribute( PolygonRef, PerimeterVertexIndices[ 2 ], UEditableMeshAttribute::VertexColor(), 0 )
+					TriangleVertexWeights.X * GetVertexInstanceAttribute( Triangle.VertexInstanceIDs[ 0 ], UEditableMeshAttribute::VertexColor(), 0 ) +
+					TriangleVertexWeights.Y * GetVertexInstanceAttribute( Triangle.VertexInstanceIDs[ 1 ], UEditableMeshAttribute::VertexColor(), 0 ) +
+					TriangleVertexWeights.Z * GetVertexInstanceAttribute( Triangle.VertexInstanceIDs[ 2 ], UEditableMeshAttribute::VertexColor(), 0 )
 				);
 			}
 			else
@@ -4759,7 +6739,11 @@
 				const FVertexID RightInsetVertexID = NewVertexIDs[ RightVertexNumber ];
 
 				FPolygonToCreate& PolygonToCreate = *new( SidePolygonsToCreate ) FPolygonToCreate();
-				PolygonToCreate.SectionID = PolygonRef.SectionID;
+				PolygonToCreate.PolygonGroupID = PolygonGroupID;
+				if( bShouldBevel )
+				{
+					PolygonToCreate.PolygonEdgeHardness = EPolygonEdgeHardness::AllEdgesHard;
+				}
 
 				// Original left
 				{
@@ -4772,13 +6756,13 @@
 						NewPerimeterVertex.PolygonVertexAttributes.Attributes.Add( FMeshElementAttributeData(
 							UEditableMeshAttribute::VertexTextureCoordinate(),
 							TextureCoordinateIndex,
-							GetPolygonPerimeterVertexAttribute( PolygonRef, LeftVertexNumber, UEditableMeshAttribute::VertexTextureCoordinate(), TextureCoordinateIndex ) ) );
+							GetPolygonPerimeterVertexAttribute( PolygonID, LeftVertexNumber, UEditableMeshAttribute::VertexTextureCoordinate(), TextureCoordinateIndex ) ) );
 					}
 
 					NewPerimeterVertex.PolygonVertexAttributes.Attributes.Add( FMeshElementAttributeData(
 						UEditableMeshAttribute::VertexColor(),
 						0,
-						GetPolygonPerimeterVertexAttribute( PolygonRef, LeftVertexNumber, UEditableMeshAttribute::VertexColor(), 0 ) ) );
+						GetPolygonPerimeterVertexAttribute( PolygonID, LeftVertexNumber, UEditableMeshAttribute::VertexColor(), 0 ) ) );
 				}
 
 				// Original right
@@ -4792,13 +6776,13 @@
 						NewPerimeterVertex.PolygonVertexAttributes.Attributes.Add( FMeshElementAttributeData(
 							UEditableMeshAttribute::VertexTextureCoordinate(),
 							TextureCoordinateIndex,
-							GetPolygonPerimeterVertexAttribute( PolygonRef, RightVertexNumber, UEditableMeshAttribute::VertexTextureCoordinate(), TextureCoordinateIndex ) ) );
+							GetPolygonPerimeterVertexAttribute( PolygonID, RightVertexNumber, UEditableMeshAttribute::VertexTextureCoordinate(), TextureCoordinateIndex ) ) );
 					}
 
 					NewPerimeterVertex.PolygonVertexAttributes.Attributes.Add( FMeshElementAttributeData(
 						UEditableMeshAttribute::VertexColor(),
 						0,
-						GetPolygonPerimeterVertexAttribute( PolygonRef, RightVertexNumber, UEditableMeshAttribute::VertexColor(), 0 ) ) );
+						GetPolygonPerimeterVertexAttribute( PolygonID, RightVertexNumber, UEditableMeshAttribute::VertexColor(), 0 ) ) );
 				}
 
 				// Inset right
@@ -4847,7 +6831,11 @@
 		if( Mode == EInsetPolygonsMode::All || Mode == EInsetPolygonsMode::CenterPolygonOnly )
 		{
 			FPolygonToCreate& PolygonToCreate = *new( CenterPolygonsToCreate ) FPolygonToCreate();
-			PolygonToCreate.SectionID = PolygonRef.SectionID;
+			PolygonToCreate.PolygonGroupID = PolygonGroupID;
+			if( bShouldBevel )
+			{
+				PolygonToCreate.PolygonEdgeHardness = EPolygonEdgeHardness::AllEdgesHard;
+			}
 
 			for( int32 NewVertexNumber = 0; NewVertexNumber < NewVertexIDs.Num(); ++NewVertexNumber )
 			{
@@ -4877,8 +6865,9 @@
 	{
 		const bool bDeleteOrphanedEdges = false;	// No need to delete orphans, because this function won't orphan anything
 		const bool bDeleteOrphanedVertices = false;
-		const bool bDeleteEmptySections = false;
-		DeletePolygons( PolygonRefs, bDeleteOrphanedEdges, bDeleteOrphanedVertices, bDeleteEmptySections );
+		const bool bDeleteOrphanedVertexInstances = false;
+		const bool bDeleteEmptyPolygonGroups = false;
+		DeletePolygons( PolygonIDs, bDeleteOrphanedEdges, bDeleteOrphanedVertices, bDeleteOrphanedVertexInstances, bDeleteEmptyPolygonGroups );
 	}
 
 	// @todo mesheditor inset: Should we create the new inset polygon with hard edges or soft?  Currently using the default.
@@ -4889,203 +6878,233 @@
 		SetVerticesAttributes( AttributesForVertices );
 	}
 
-	// Create the new side polygons
-	static TArray<FPolygonRef> AllNewPolygonRefs;
-	AllNewPolygonRefs.Reset();
-
-	static TArray<FEdgeID> EdgesToMakeHard;
-	EdgesToMakeHard.Reset();
-
-	static TArray<bool> EdgesNewIsHard;
-	EdgesNewIsHard.Reset();
-
 	if( SidePolygonsToCreate.Num() > 0 )
 	{
 		static TArray<FEdgeID> NewEdgeIDs;
-		CreatePolygons( SidePolygonsToCreate, /* Out */ OutNewSidePolygonRefs, /* Out */ NewEdgeIDs );
-
-		AllNewPolygonRefs.Append( OutNewSidePolygonRefs );
-
-		// If we're beveling, go ahead and make all of the side edges hard
-		if( bShouldBevel )
-		{
-			static TSet<FEdgeID> AllSidePolygonEdgeIDs;
-			AllSidePolygonEdgeIDs.Reset();
-			for( const FPolygonRef PolygonRef : OutNewSidePolygonRefs )
-			{
-				static TArray<FEdgeID> EdgeIDs;
-				GetPolygonPerimeterEdges( PolygonRef, /* Out */ EdgeIDs );
-
-				AllSidePolygonEdgeIDs.Append( EdgeIDs );
-			}
-
-			for( const FEdgeID EdgeID : AllSidePolygonEdgeIDs )
-			{
-				EdgesToMakeHard.Add( EdgeID );
-				EdgesNewIsHard.Add( true );
-			}
-		}
+		CreatePolygons( SidePolygonsToCreate, /* Out */ OutNewSidePolygonIDs, /* Out */ NewEdgeIDs );
 	}
 
 	// Create the new center polygons.  Note that we pass back the IDs of the newly-created center polygons
 	if( CenterPolygonsToCreate.Num() > 0 )
 	{
 		static TArray<FEdgeID> NewEdgeIDs;
-		CreatePolygons( CenterPolygonsToCreate, /* Out */ OutNewCenterPolygonRefs, /* Out */ NewEdgeIDs );
-
-		AllNewPolygonRefs.Append( OutNewCenterPolygonRefs );
-
-		// Make the perimeter edges of the new polygons hard
-		static TSet<FEdgeID> PerimeterEdgeIDsToMakeHard;
-		PerimeterEdgeIDsToMakeHard.Reset();
-		for( const FPolygonRef NewCenterPolygonRef : OutNewCenterPolygonRefs )
-		{
-			static TArray<FEdgeID> PerimeterEdgeIDs;
-			GetPolygonPerimeterEdges( NewCenterPolygonRef, /* Out */ PerimeterEdgeIDs );
-
-			for( const FEdgeID EdgeID : PerimeterEdgeIDs )
-			{
-				PerimeterEdgeIDsToMakeHard.Add( EdgeID );
-			}
-		}
-
-		for( const FEdgeID EdgeID : PerimeterEdgeIDsToMakeHard )
-		{
-			EdgesToMakeHard.Add( EdgeID );
-			EdgesNewIsHard.Add( true );
-		}
-	}
-
-	if( EdgesToMakeHard.Num() > 0 )
-	{
-		SetEdgesHardness( EdgesToMakeHard, EdgesNewIsHard );
-	}
-
-	// Generate normals for all of the new polygons
-	if( bShouldBevel )
-	{
-		// When beveling, we're moving the polygon perimeter vertex locations which will certainly
-		// affect the normals and tangents of adjacent polygons as well!
-		GenerateNormalsAndTangentsForPolygonsAndAdjacents( AllNewPolygonRefs );
-	}
-	else
-	{
-		// When not beveling, the original polygon's perimeter stays the same, so we don't need to
-		// worry about normals for adjacent polygons
-		GenerateNormalsAndTangentsForPolygons( AllNewPolygonRefs );
+		CreatePolygons( CenterPolygonsToCreate, /* Out */ OutNewCenterPolygonIDs, /* Out */ NewEdgeIDs );
 	}
 }
 
 
-void UEditableMesh::InsetPolygons( const TArray<FPolygonRef>& PolygonRefs, const float InsetFixedDistance, const float InsetProgressTowardCenter, const EInsetPolygonsMode Mode, TArray<FPolygonRef>& OutNewCenterPolygonRefs, TArray<FPolygonRef>& OutNewSidePolygonRefs )
+void UEditableMesh::InsetPolygons( const TArray<FPolygonID>& PolygonIDs, const float InsetFixedDistance, const float InsetProgressTowardCenter, const EInsetPolygonsMode Mode, TArray<FPolygonID>& OutNewCenterPolygonIDs, TArray<FPolygonID>& OutNewSidePolygonIDs )
 {
 	const bool bShouldBevel = false;
-	BevelOrInsetPolygons( PolygonRefs, InsetFixedDistance, InsetProgressTowardCenter, Mode, bShouldBevel, /* Out */ OutNewCenterPolygonRefs, /* Out */ OutNewSidePolygonRefs );
+	BevelOrInsetPolygons( PolygonIDs, InsetFixedDistance, InsetProgressTowardCenter, Mode, bShouldBevel, /* Out */ OutNewCenterPolygonIDs, /* Out */ OutNewSidePolygonIDs );
 }
 
 
-void UEditableMesh::BevelPolygons( const TArray<FPolygonRef>& PolygonRefs, const float BevelFixedDistance, const float BevelProgressTowardCenter, TArray<FPolygonRef>& OutNewCenterPolygonRefs, TArray<FPolygonRef>& OutNewSidePolygonRefs )
+void UEditableMesh::BevelPolygons( const TArray<FPolygonID>& PolygonIDs, const float BevelFixedDistance, const float BevelProgressTowardCenter, TArray<FPolygonID>& OutNewCenterPolygonIDs, TArray<FPolygonID>& OutNewSidePolygonIDs )
 {
 	const bool bShouldBevel = true;
-	BevelOrInsetPolygons( PolygonRefs, BevelFixedDistance, BevelProgressTowardCenter, EInsetPolygonsMode::All, bShouldBevel, /* Out */ OutNewCenterPolygonRefs, /* Out */ OutNewSidePolygonRefs );
+	BevelOrInsetPolygons( PolygonIDs, BevelFixedDistance, BevelProgressTowardCenter, EInsetPolygonsMode::All, bShouldBevel, /* Out */ OutNewCenterPolygonIDs, /* Out */ OutNewSidePolygonIDs );
 }
 
 
-void UEditableMesh::GenerateNormalsAndTangentsForPolygons( const TArray< FPolygonRef >& PolygonRefs )
+float UEditableMesh::GetPolygonCornerAngleForVertex( const FPolygonID PolygonID, const FVertexID VertexID ) const
 {
-	static TArray<FVertexAttributesForPolygon> VertexAttributesForPolygons;
-	VertexAttributesForPolygons.Reset();
-	VertexAttributesForPolygons.Reserve( PolygonRefs.Num() );
+	const FMeshPolygon& Polygon = Polygons[ PolygonID.GetValue() ];
 
-	for( const FPolygonRef& PolygonRef : PolygonRefs )
+	// Lambda function which returns the inner angle at a given index on a polygon contour
+	auto GetContourAngle = [ this ]( const FMeshPolygonContour& Contour, const int32 ContourIndex )
 	{
-		FVertexAttributesForPolygon& PolygonNewAttributes = *new( VertexAttributesForPolygons ) FVertexAttributesForPolygon();
-		PolygonNewAttributes.PolygonRef = PolygonRef;
+		const int32 NumVertices = Contour.VertexInstanceIDs.Num();
 
-		const int32 PolygonPerimeterVertexCount = GetPolygonPerimeterVertexCount( PolygonRef );
+		const FVector PrevVertexPosition = GetVertexInstanceAttribute( Contour.VertexInstanceIDs[ ( ContourIndex + NumVertices - 1 ) % NumVertices ], UEditableMeshAttribute::VertexPosition(), 0 );
+		const FVector ThisVertexPosition = GetVertexInstanceAttribute( Contour.VertexInstanceIDs[ ContourIndex ], UEditableMeshAttribute::VertexPosition(), 0 );
+		const FVector NextVertexPosition = GetVertexInstanceAttribute( Contour.VertexInstanceIDs[ ( ContourIndex + 1 ) % NumVertices ], UEditableMeshAttribute::VertexPosition(), 0 );
 
-		PolygonNewAttributes.PerimeterVertexAttributeLists.SetNum( PolygonPerimeterVertexCount, false );
+		const FVector Direction1 = ( PrevVertexPosition - ThisVertexPosition ).GetSafeNormal();
+		const FVector Direction2 = ( NextVertexPosition - ThisVertexPosition ).GetSafeNormal();
 
-		for( int32 VertexNumber = 0; VertexNumber < PolygonPerimeterVertexCount; ++VertexNumber )
+		return FMath::Acos( FVector::DotProduct( Direction1, Direction2 ) );
+	};
+
+	auto IsVertexInstancedFromThisVertex = [ this, VertexID ]( const FVertexInstanceID VertexInstanceID )
+	{
+		return VertexInstances[ VertexInstanceID.GetValue() ].VertexID == VertexID;
+	};
+
+	// First look for the vertex instance in the perimeter
+	int32 ContourIndex = Polygon.PerimeterContour.VertexInstanceIDs.IndexOfByPredicate( IsVertexInstancedFromThisVertex );
+	if( ContourIndex != INDEX_NONE )
+	{
+		// Return the internal angle if found
+		return GetContourAngle( Polygon.PerimeterContour, ContourIndex );
+	}
+	else
+	{
+		// If not found, look in all the holes
+		for( const FMeshPolygonContour& HoleContour : Polygon.HoleContours )
 		{
-			// Compute the vertex normal
-			const FVector VertexNormal = ComputePolygonPerimeterVertexNormal( PolygonRef, VertexNumber );
-
-			// Save it!
-			// @todo mesheditor urgent perf: This causes every polygon vertex to be discreet, even if the vertex shares the same normal/uvs/tangents between coincident rendering vertices (common for smooth meshes)
-			TArray<FMeshElementAttributeData>& PerimeterVertexNewAttributes = PolygonNewAttributes.PerimeterVertexAttributeLists[ VertexNumber ].Attributes;
-			FMeshElementAttributeData& PerimeterVertexAttribute = *new( PerimeterVertexNewAttributes ) FMeshElementAttributeData( UEditableMeshAttribute::VertexNormal(), 0, FVector4( VertexNormal, 0.0f ) );
-		}
-	}
-	SetPolygonsVertexAttributes( VertexAttributesForPolygons );
-
-	// Now update the tangents
-	GenerateTangentsForPolygons( PolygonRefs );
-}
-
-
-void UEditableMesh::GenerateNormalsAndTangentsForPolygonsAndAdjacents( const TArray< FPolygonRef >& PolygonRefs )
-{
-	// Figure out everything next to these polygons
-	static TArray<FPolygonRef> PolygonRefsWithAdjacents;
-	PolygonRefsWithAdjacents.Reset();
-
-	PolygonRefsWithAdjacents = PolygonRefs;
-
-	for( const FPolygonRef PolygonRef : PolygonRefs )
-	{
-		static TArray<FVertexID> PerimeterVertexIDs;
-		PerimeterVertexIDs.Reset();
-		GetPolygonPerimeterVertices( PolygonRef, /* Out */ PerimeterVertexIDs );
-
-		for( const FVertexID PerimeterVertexID : PerimeterVertexIDs )
-		{
-			static TArray<FPolygonRef> ConnectedPolygonRefs;
-			GetVertexConnectedPolygons( PerimeterVertexID, /* Out */ ConnectedPolygonRefs );
-			for( const FPolygonRef ConnectedPolygonRef : ConnectedPolygonRefs )
+			ContourIndex = HoleContour.VertexInstanceIDs.IndexOfByPredicate( IsVertexInstancedFromThisVertex );
+			if( ContourIndex != INDEX_NONE )
 			{
-				// Ignore self
-				if( ConnectedPolygonRef != PolygonRef )
-				{
-					// Add uniquely, because it may have been one of the polygons that were passed in, or even a shared 
-					// neighbor of one of those polygons
-					PolygonRefsWithAdjacents.AddUnique( ConnectedPolygonRef );
-				}
+				// Hole vertex contribution is the part which ISN'T the internal angle of the contour, so subtract from 2*pi
+				return ( 2.0f * PI ) - GetContourAngle( HoleContour, ContourIndex );
 			}
 		}
 	}
 
-	return GenerateNormalsAndTangentsForPolygons( PolygonRefsWithAdjacents );
+	// Found nothing; return 0
+	return 0.0f;
 }
 
 
-void UEditableMesh::GenerateTangentsForPolygons( const TArray< FPolygonRef >& PolygonRefs )
+void UEditableMesh::GeneratePolygonTangentsAndNormals( const TArray<FPolygonID>& PolygonIDs )
 {
+	for( const FPolygonID PolygonID : PolygonIDs )
+	{
+		FMeshPolygon& Polygon = Polygons[ PolygonID.GetValue() ];
+
+		// Calculate the tangent basis for the polygon, based on the average of all constituent triangles
+		FVector Normal = FVector::ZeroVector;
+		FVector Tangent = FVector::ZeroVector;
+		FVector Binormal = FVector::ZeroVector;
+
+		for( const FMeshTriangle& Triangle : Polygon.Triangles )
+		{
+			const FMeshVertexInstance& VertexInstance0 = VertexInstances[ Triangle.VertexInstanceIDs[ 0 ].GetValue() ];
+			const FMeshVertexInstance& VertexInstance1 = VertexInstances[ Triangle.VertexInstanceIDs[ 1 ].GetValue() ];
+			const FMeshVertexInstance& VertexInstance2 = VertexInstances[ Triangle.VertexInstanceIDs[ 2 ].GetValue() ];
+
+			const FMeshVertex& Vertex0 = Vertices[ VertexInstance0.VertexID.GetValue() ];
+			const FMeshVertex& Vertex1 = Vertices[ VertexInstance1.VertexID.GetValue() ];
+			const FMeshVertex& Vertex2 = Vertices[ VertexInstance2.VertexID.GetValue() ];
+
+			const FVector DPosition1 = Vertex1.VertexPosition - Vertex0.VertexPosition;
+			const FVector DPosition2 = Vertex2.VertexPosition - Vertex0.VertexPosition;
+
+			// @todo mesheditor: currently hardcoded to calculate the tangent basis from UV0.
+			const FVector2D DUV1 = VertexInstance1.VertexUVs[ 0 ] - VertexInstance0.VertexUVs[ 0 ];
+			const FVector2D DUV2 = VertexInstance2.VertexUVs[ 0 ] - VertexInstance0.VertexUVs[ 0 ];
+
+			// We have a left-handed coordinate system, but a counter-clockwise winding order
+			// Hence normal calculation has to take the triangle vectors cross product in reverse.
+			Normal += FVector::CrossProduct( DPosition2, DPosition1 );
+
+			// ...and tangent space seems to be right-handed.
+			const float DetUV = FVector2D::CrossProduct( DUV1, DUV2 );
+			const float InvDetUV = ( DetUV == 0.0f ) ? 0.0f : 1.0f / DetUV;
+
+			Tangent += ( DPosition1 * DUV2.Y - DPosition2 * DUV1.Y ) * InvDetUV;
+			Binormal += ( DPosition2 * DUV1.X - DPosition1 * DUV2.X ) * InvDetUV;
+		}
+
+		Polygon.PolygonNormal = Normal.GetSafeNormal();
+		Polygon.PolygonTangent = Tangent.GetSafeNormal();
+		Polygon.PolygonBinormal = Binormal.GetSafeNormal();
+	}
+}
+
+
+void UEditableMesh::GenerateTangentsAndNormals()
+{
+	GeneratePolygonTangentsAndNormals( PolygonsPendingNewTangentBasis.Array() );
+
+	static TSet<FVertexInstanceID> VertexInstanceIDs;
+	VertexInstanceIDs.Reset();
+
+	for( const FPolygonID PolygonID : PolygonsPendingNewTangentBasis )
+	{
+		const FMeshPolygon& Polygon = Polygons[ PolygonID.GetValue() ];
+
+		VertexInstanceIDs.Append( Polygon.PerimeterContour.VertexInstanceIDs );
+		for( const FMeshPolygonContour& HoleContour : Polygon.HoleContours )
+		{
+			VertexInstanceIDs.Append( HoleContour.VertexInstanceIDs );
+		}
+	}
+
+	for( const FVertexInstanceID VertexInstanceID : VertexInstanceIDs )
+	{
+		FVector Normal = FVector::ZeroVector;
+		FVector Tangent = FVector::ZeroVector;
+		FVector Binormal = FVector::ZeroVector;
+
+		FMeshVertexInstance& VertexInstance = VertexInstances[ VertexInstanceID.GetValue() ];
+		const FVertexID VertexID = VertexInstance.VertexID;
+
+		// Get all polygons connected to this vertex instance, and also any in the same smoothing group connected to a different vertex instance
+		// (as they still have influence over the normal).
+		static TArray<FPolygonID> AllConnectedPolygons;
+		check( VertexInstance.ConnectedPolygons.Num() > 0 );
+		GetPolygonsInSameSoftEdgedGroup( VertexID, VertexInstance.ConnectedPolygons[ 0 ], AllConnectedPolygons );
+
+		// The vertex instance normal is computed as a sum of all connected polygons' normals, weighted by the angle they make with the vertex
+		for( const FPolygonID ConnectedPolygonID : AllConnectedPolygons )
+		{
+			float Angle = GetPolygonCornerAngleForVertex( ConnectedPolygonID, VertexID );
+
+			const FMeshPolygon& Polygon = Polygons[ ConnectedPolygonID.GetValue() ];
+			Normal += Polygon.PolygonNormal * Angle;
+
+			// If this polygon is actually connected to the vertex instance we're processing, also include its contributions towards the tangent
+			if( VertexInstance.ConnectedPolygons.Contains( ConnectedPolygonID ) )
+			{
+				Tangent += Polygon.PolygonTangent * Angle;
+				Binormal += Polygon.PolygonBinormal * Angle;
+			}
+		}
+
+		// Normalize Normal
+		Normal = Normal.GetSafeNormal();
+
+		// Make Tangent orthonormal to Normal.
+		// This is a quicker method than normalizing Tangent, taking the cross product Normal X Tangent, and then a further cross product with that result
+		Tangent = ( Tangent - Normal * FVector::DotProduct( Normal, Tangent ) ).GetSafeNormal();
+
+		// Calculate binormal sign
+		const float BinormalSign = ( FVector::DotProduct( FVector::CrossProduct( Normal, Tangent ), Binormal ) < 0.0f ) ? -1.0f : 1.0f;
+
+		//UE_LOG( LogEditableMesh, Verbose, TEXT( "%s normal: old %s" ), *VertexInstanceID.ToString(), *VertexInstance.Normal.ToString() );
+		//UE_LOG( LogEditableMesh, Verbose, TEXT( "%s normal: new %s" ), *VertexInstanceID.ToString(), *Normal.ToString() );
+
+		//UE_LOG( LogEditableMesh, Verbose, TEXT( "%s tangent: old %s" ), *VertexInstanceID.ToString(), *VertexInstance.Tangent.ToString() );
+		//UE_LOG( LogEditableMesh, Verbose, TEXT( "%s tangent: new %s" ), *VertexInstanceID.ToString(), *Tangent.ToString() );
+
+		VertexInstance.Normal = Normal;
+		VertexInstance.Tangent = Tangent;
+		VertexInstance.BinormalSign = BinormalSign;
+	}
+}
+
+
+#if 0
+void UEditableMesh::GenerateTangentsForPolygons( const TArray< FPolygonID >& PolygonIDs )
+{
+	// @todo mesheditor: This method is currently not used. Revisit MikkTSpace in the future.
 	static TArray<FVertexAttributesForPolygon> VertexAttributesForPolygons;
 	VertexAttributesForPolygons.Reset();
-	VertexAttributesForPolygons.Reserve( PolygonRefs.Num() );
-	for( const FPolygonRef& PolygonRef : PolygonRefs )
+	VertexAttributesForPolygons.Reserve( PolygonIDs.Num() );
+	for( const FPolygonID PolygonID : PolygonIDs )
 	{
 		FVertexAttributesForPolygon& PolygonNewAttributes = *new( VertexAttributesForPolygons ) FVertexAttributesForPolygon();
-		PolygonNewAttributes.PolygonRef = PolygonRef;
+		PolygonNewAttributes.PolygonID = PolygonID;
 
-		const int32 PolygonPerimeterVertexCount = GetPolygonPerimeterVertexCount( PolygonRef );
+		const int32 PolygonPerimeterVertexCount = GetPolygonPerimeterVertexCount( PolygonID );
 		PolygonNewAttributes.PerimeterVertexAttributeLists.SetNum( PolygonPerimeterVertexCount, false );
 	}
 
 	struct FMikkUserData
 	{
 		UEditableMesh* Self;
-		const TArray< FPolygonRef >& Polygons;
+		const TArray< FPolygonID >& Polygons;
 		TArray<FVertexAttributesForPolygon>& VertexAttributesForPolygons;
 
-		FMikkUserData( UEditableMesh* InitSelf, const TArray< FPolygonRef >& InitPolygons, TArray<FVertexAttributesForPolygon>& InitVertexAttributesForPolygons )
+		FMikkUserData( UEditableMesh* InitSelf, const TArray< FPolygonID >& InitPolygons, TArray<FVertexAttributesForPolygon>& InitVertexAttributesForPolygons )
 			: Self( InitSelf ),
 			  Polygons( InitPolygons ),
 			  VertexAttributesForPolygons( InitVertexAttributesForPolygons )
 		{
 		}
-	} MikkUserData( this, PolygonRefs, VertexAttributesForPolygons );
+	} MikkUserData( this, PolygonIDs, VertexAttributesForPolygons );
 
 	struct Local
 	{
@@ -5099,10 +7118,10 @@
 		{
 			const FMikkUserData& UserData = *static_cast<FMikkUserData*>( Context->m_pUserData );
 
-			const FPolygonRef& PolygonRef = UserData.Polygons[ MikkFaceIndex ];
+			const FPolygonID PolygonID = UserData.Polygons[ MikkFaceIndex ];
 
 			// @todo mesheditor holes: Can we even use Mikktpsace for polygons with holes?  Do we need to run the triangulated triangles through instead?
-			const int32 PolygonVertexCount = UserData.Self->GetPolygonPerimeterVertexCount( PolygonRef );
+			const int32 PolygonVertexCount = UserData.Self->GetPolygonPerimeterVertexCount( PolygonID );
 			return PolygonVertexCount;
 		}
 
@@ -5110,8 +7129,8 @@
 		{
 			const FMikkUserData& UserData = *static_cast<FMikkUserData*>( Context->m_pUserData );
 
-			const FPolygonRef& PolygonRef = UserData.Polygons[ MikkFaceIndex ];
-			const FVector VertexPosition = UserData.Self->GetPolygonPerimeterVertexAttribute( PolygonRef, MikkVertexIndex, UEditableMeshAttribute::VertexPosition(), 0 );
+			const FPolygonID PolygonID = UserData.Polygons[ MikkFaceIndex ];
+			const FVector VertexPosition = UserData.Self->GetPolygonPerimeterVertexAttribute( PolygonID, MikkVertexIndex, UEditableMeshAttribute::VertexPosition(), 0 );
 
 			OutPosition[0] = VertexPosition.X;
 			OutPosition[1] = VertexPosition.Y;
@@ -5122,8 +7141,8 @@
 		{
 			const FMikkUserData& UserData = *static_cast<FMikkUserData*>( Context->m_pUserData );
 
-			const FPolygonRef& PolygonRef = UserData.Polygons[ MikkFaceIndex ];
-			const FVector PolygonVertexNormal = UserData.Self->GetPolygonPerimeterVertexAttribute( PolygonRef, MikkVertexIndex, UEditableMeshAttribute::VertexNormal(), 0 );
+			const FPolygonID PolygonID = UserData.Polygons[ MikkFaceIndex ];
+			const FVector PolygonVertexNormal = UserData.Self->GetPolygonPerimeterVertexAttribute( PolygonID, MikkVertexIndex, UEditableMeshAttribute::VertexNormal(), 0 );
 
 			OutNormal[ 0 ] = PolygonVertexNormal.X;
 			OutNormal[ 1 ] = PolygonVertexNormal.Y;
@@ -5137,8 +7156,8 @@
 			// @todo mesheditor: Support using a custom texture coordinate index for tangent space generation?
 			const int32 TextureCoordinateIndex = 0;
 
-			const FPolygonRef& PolygonRef = UserData.Polygons[ MikkFaceIndex ];
-			const FVector2D PolygonVertexTextureCoordinate( UserData.Self->GetPolygonPerimeterVertexAttribute( PolygonRef, MikkVertexIndex, UEditableMeshAttribute::VertexTextureCoordinate(), TextureCoordinateIndex ) );
+			const FPolygonID PolygonID = UserData.Polygons[ MikkFaceIndex ];
+			const FVector2D PolygonVertexTextureCoordinate( UserData.Self->GetPolygonPerimeterVertexAttribute( PolygonID, MikkVertexIndex, UEditableMeshAttribute::VertexTextureCoordinate(), TextureCoordinateIndex ) );
 
 			OutUV[0] = PolygonVertexTextureCoordinate.X;
 			OutUV[1] = PolygonVertexTextureCoordinate.Y;
@@ -5184,6 +7203,7 @@
 
 	SetPolygonsVertexAttributes( VertexAttributesForPolygons );
 }
+#endif
 
 
 void UEditableMesh::SetVerticesCornerSharpness( const TArray<FVertexID>& VertexIDs, const TArray<float>& VerticesNewSharpness )
@@ -5239,8 +7259,8 @@
 	static TArray<FAttributesForEdge> AttributesForEdges;
 	AttributesForEdges.Reset();
 
-	static TSet<FPolygonRef> UniqueConnectedPolygonRefs;
-	UniqueConnectedPolygonRefs.Reset();
+	static TSet<FPolygonID> UniqueConnectedPolygonIDs;
+	UniqueConnectedPolygonIDs.Reset();
 
 	for( int32 EdgeNumber = 0; EdgeNumber < EdgeIDs.Num(); ++EdgeNumber )
 	{
@@ -5254,19 +7274,12 @@
 			FVector4( EdgesNewIsHard[ EdgeNumber ] ? 1.0f : 0.0f ) ) );
 
 		// Get the polygons this edge is connected to.  They'll need new normals.
-		static TArray<FPolygonRef> ConnectedPolygonRefs;
-		GetEdgeConnectedPolygons( EdgeID, /* Out */ ConnectedPolygonRefs );
-
-		for( const FPolygonRef ConnectedPolygonRef : ConnectedPolygonRefs )
-		{
-			UniqueConnectedPolygonRefs.Add( ConnectedPolygonRef );
-		}
+		static TArray<FPolygonID> ConnectedPolygonIDs;
+		GetEdgeConnectedPolygons( EdgeID, /* Out */ ConnectedPolygonIDs );
+		UniqueConnectedPolygonIDs.Append( ConnectedPolygonIDs );
 	}
 
 	SetEdgesAttributes( AttributesForEdges );
-
-	// Generate fresh normals
-	GenerateNormalsAndTangentsForPolygons( MoveTemp( UniqueConnectedPolygonRefs.Array() ) );
 }
 
 
@@ -5287,15 +7300,15 @@
 		{
 			float MinDot = 1.0f;
 
-			const FPolygonRef FirstPolygonRef = GetEdgeConnectedPolygon( EdgeID, 0 );
+			const FPolygonID FirstPolygonID = GetEdgeConnectedPolygon( EdgeID, 0 );
 
-			FVector LastPolygonNormal = ComputePolygonNormal( FirstPolygonRef );
+			FVector LastPolygonNormal = ComputePolygonNormal( FirstPolygonID );
 
 			for( int32 ConnectedPolygonNumber = 1; ConnectedPolygonNumber < ConnectedPolygonCount; ++ConnectedPolygonNumber )
 			{
-				const FPolygonRef PolygonRef = GetEdgeConnectedPolygon( EdgeID, ConnectedPolygonNumber );
+				const FPolygonID PolygonID = GetEdgeConnectedPolygon( EdgeID, ConnectedPolygonNumber );
 
-				const FVector PolygonNormal = ComputePolygonNormal( PolygonRef );
+				const FVector PolygonNormal = ComputePolygonNormal( PolygonID );
 
 				const float Dot = FVector::DotProduct( PolygonNormal, LastPolygonNormal );
 				MinDot = FMath::Min( Dot, MinDot );
@@ -5315,6 +7328,8 @@
 
 void UEditableMesh::SetEdgesVertices( const TArray<FVerticesForEdge>& VerticesForEdges )
 {
+	UE_LOG( LogEditableMesh, Verbose, TEXT( "SetEdgesVertices: %s" ), *LogHelpers::ArrayToString( VerticesForEdges ) );
+
 	FSetEdgesVerticesChangeInput RevertInput;
 
 	RevertInput.VerticesForEdges.AddUninitialized( VerticesForEdges.Num() );
@@ -5328,67 +7343,146 @@
 		GetEdgeVertices( VerticesForEdge.EdgeID, /* Out */ RevertVerticesForEdge.NewVertexID0, /* Out */ RevertVerticesForEdge.NewVertexID1 );
 
 		// Edit the edge
-		SetEdgeVertices_Internal( VerticesForEdge.EdgeID, VerticesForEdge.NewVertexID0, VerticesForEdge.NewVertexID1 );
+		FMeshEdge& Edge = Edges[ VerticesForEdge.EdgeID.GetValue() ];
+
+		for( uint32 EdgeVertexNumber = 0; EdgeVertexNumber < 2; ++EdgeVertexNumber )
+		{
+			// Disconnect the edge from its existing vertices
+			const FVertexID VertexID = Edge.VertexIDs[ EdgeVertexNumber ];
+			FMeshVertex& Vertex = Vertices[ VertexID.GetValue() ];
+			verify( Vertex.ConnectedEdgeIDs.RemoveSingleSwap( VerticesForEdge.EdgeID ) == 1 );	// Must have been already connected!
+		}
+
+		Edge.VertexIDs[0] = VerticesForEdge.NewVertexID0;
+		Edge.VertexIDs[1] = VerticesForEdge.NewVertexID1;
+
+		// Connect the new vertices to the edge
+		for( uint32 EdgeVertexNumber = 0; EdgeVertexNumber < 2; ++EdgeVertexNumber )
+		{
+			const FVertexID VertexID = Edge.VertexIDs[ EdgeVertexNumber ];
+			FMeshVertex& Vertex = Vertices[ VertexID.GetValue() ];
+
+			check( !Vertex.ConnectedEdgeIDs.Contains( VerticesForEdge.EdgeID ) );	// Should not have already been connected
+			Vertex.ConnectedEdgeIDs.Add( VerticesForEdge.EdgeID );
+		}
 	}
 
 	AddUndo( MakeUnique<FSetEdgesVerticesChange>( MoveTemp( RevertInput ) ) );
+
+	UE_LOG( LogEditableMesh, Verbose, TEXT( "* SetEdgesVertices finished") );
 }
 
 
-void UEditableMesh::InsertPolygonPerimeterVertices( const FPolygonRef PolygonRef, const int32 InsertBeforeVertexNumber, const TArray<FVertexAndAttributes>& VerticesToInsert )
+void UEditableMesh::InsertPolygonPerimeterVertices( const FPolygonID PolygonID, const int32 InsertBeforeVertexNumber, const TArray<FVertexAndAttributes>& VerticesToInsert )
 {
-	FRemovePolygonPerimeterVerticesChangeInput RevertInput;
+	UE_LOG( LogEditableMesh, Verbose, TEXT( "InsertPolygonPerimeterVertices: PolygonID:%s InsertBeforeVertexNumber:%d %s" ), *PolygonID.ToString(), InsertBeforeVertexNumber, *LogHelpers::ArrayToString( VerticesToInsert ) );
+	// @todo mesheditor: see if we want to keep this action. It is only used by SplitEdge, and doesn't generate missing edges.
+	// We can achieve the same thing by deleting a polygon without deleting vertex instances, and creating a new polygon with an extra vertex instance.
 
-	RevertInput.PolygonRef = PolygonRef;
-	RevertInput.FirstVertexNumberToRemove = InsertBeforeVertexNumber;
-	RevertInput.NumVerticesToRemove = VerticesToInsert.Num();
+	FMeshPolygon& Polygon = Polygons[ PolygonID.GetValue() ];
 
-	InsertPolygonPerimeterVertices_Internal( PolygonRef, InsertBeforeVertexNumber, VerticesToInsert );
+	// @todo mesheditor: should create missing edges here too?
 
-	AddUndo( MakeUnique<FRemovePolygonPerimeterVerticesChange>( MoveTemp( RevertInput ) ) );
+	// Insert new vertex instances
+	for( int32 InsertVertexIter = 0; InsertVertexIter < VerticesToInsert.Num(); ++InsertVertexIter )
+	{
+		const FVertexAndAttributes& VertexToInsert = VerticesToInsert[ InsertVertexIter ];
+
+		const FVertexInstanceID VertexInstanceID = CreateVertexInstanceForContourVertex( VertexToInsert, PolygonID );
+		Polygon.PerimeterContour.VertexInstanceIDs.Insert( VertexInstanceID, InsertBeforeVertexNumber + InsertVertexIter );
+
+		// Connect polygon to vertex instance
+		FMeshVertexInstance& VertexInstance = VertexInstances[ VertexInstanceID.GetValue() ];
+		check( !VertexInstance.ConnectedPolygons.Contains( PolygonID ) );
+		VertexInstance.ConnectedPolygons.Add( PolygonID );
+	}
+
+	// Back up insert
+	{
+		FRemovePolygonPerimeterVerticesChangeInput RevertInput;
+
+		RevertInput.PolygonID = PolygonID;
+		RevertInput.FirstVertexNumberToRemove = InsertBeforeVertexNumber;
+		RevertInput.NumVerticesToRemove = VerticesToInsert.Num();
+		RevertInput.bDeleteOrphanedVertexInstances = false;
+
+		AddUndo( MakeUnique<FRemovePolygonPerimeterVerticesChange>( MoveTemp( RevertInput ) ) );
+	}
+
+	PolygonsPendingTriangulation.Add( PolygonID );
+	PolygonsPendingNewTangentBasis.Add( PolygonID );	// @todo mesheditor: and other vertex connected polygons too?
+
+	UE_LOG( LogEditableMesh, Verbose, TEXT( "* InsertPolygonPerimeterVertices finished" ) );
 }
 
 
-void UEditableMesh::RemovePolygonPerimeterVertices( const FPolygonRef PolygonRef, const int32 FirstVertexNumberToRemove, const int32 NumVerticesToRemove )
+void UEditableMesh::RemovePolygonPerimeterVertices( const FPolygonID PolygonID, const int32 FirstVertexNumberToRemove, const int32 NumVerticesToRemove, const bool bDeleteOrphanedVertexInstances )
 {
-	FInsertPolygonPerimeterVerticesChangeInput RevertInput;
+	UE_LOG( LogEditableMesh, Verbose, TEXT( "RemovePolygonPerimeterVertices: PolygonID:%s FirstVertexNumberToRemove:%d NumVerticesToRemove:%d" ), *PolygonID.ToString(), FirstVertexNumberToRemove, NumVerticesToRemove );
+	FMeshPolygon& Polygon = Polygons[ PolygonID.GetValue() ];
 
-	RevertInput.PolygonRef = PolygonRef;
-	RevertInput.InsertBeforeVertexNumber = FirstVertexNumberToRemove;
+	// Back up 
+	{
+		FInsertPolygonPerimeterVerticesChangeInput RevertInput;
 
-	RevertInput.VerticesToInsert.SetNum( NumVerticesToRemove, false );
-	for( int32 VertexToRemoveIter = 0; VertexToRemoveIter < NumVerticesToRemove; ++VertexToRemoveIter )
+		RevertInput.PolygonID = PolygonID;
+		RevertInput.InsertBeforeVertexNumber = FirstVertexNumberToRemove;
+
+		RevertInput.VerticesToInsert.SetNum( NumVerticesToRemove, false );
+		for( int32 VertexToRemoveIter = 0; VertexToRemoveIter < NumVerticesToRemove; ++VertexToRemoveIter )
+		{
+			FVertexAndAttributes& RevertVertexToInsert = RevertInput.VerticesToInsert[ VertexToRemoveIter ];
+			RevertVertexToInsert.VertexInstanceID = Polygon.PerimeterContour.VertexInstanceIDs[ FirstVertexNumberToRemove + VertexToRemoveIter ];
+			RevertVertexToInsert.VertexID = FVertexID::Invalid;
+		}
+
+		AddUndo( MakeUnique<FInsertPolygonPerimeterVerticesChange>( MoveTemp( RevertInput ) ) );
+	}
+
 	{
-		FVertexAndAttributes& RevertVertexToInsert = RevertInput.VerticesToInsert[ VertexToRemoveIter ];
-		RevertVertexToInsert.VertexID = GetPolygonPerimeterVertex( PolygonRef, FirstVertexNumberToRemove + VertexToRemoveIter );
+		static TArray<FVertexInstanceID> OrphanedVertexInstanceIDs;
+		OrphanedVertexInstanceIDs.Reset();
 
-		for( const FName AttributeName : UEditableMesh::GetValidPolygonVertexAttributes() )
+		// Delete them backwards so it is done in the opposite order to Insert
+		for( int32 Index = NumVerticesToRemove - 1; Index >= 0; --Index )
 		{
-			const int32 MaxAttributeIndex = GetMaxAttributeIndex( AttributeName );
-			for( int32 AttributeIndex = 0; AttributeIndex < MaxAttributeIndex; ++AttributeIndex )
+			const int32 VertexNumber = FirstVertexNumberToRemove + Index;
+			const FVertexInstanceID ContourVertexInstanceID = Polygon.PerimeterContour.VertexInstanceIDs[ VertexNumber ];
+
+			// Disconnect the polygon from the vertex instance
+			FMeshVertexInstance& VertexInstance = VertexInstances[ ContourVertexInstanceID.GetValue() ];
+			verify( VertexInstance.ConnectedPolygons.Remove( PolygonID ) == 1 );
+
+			// If the vertex instance is now orphaned, add it to the list
+			if( bDeleteOrphanedVertexInstances && VertexInstance.ConnectedPolygons.Num() == 0 )
 			{
-				FMeshElementAttributeData& PolygonVertexAttribute = *new( RevertVertexToInsert.PolygonVertexAttributes.Attributes ) FMeshElementAttributeData(
-					AttributeName,
-					AttributeIndex,
-					GetPolygonPerimeterVertexAttribute( PolygonRef, FirstVertexNumberToRemove + VertexToRemoveIter, AttributeName, AttributeIndex ) );
+				OrphanedVertexInstanceIDs.Add( ContourVertexInstanceID );
 			}
 		}
+
+		Polygon.PerimeterContour.VertexInstanceIDs.RemoveAt( FirstVertexNumberToRemove, NumVerticesToRemove );
+
+		if( OrphanedVertexInstanceIDs.Num() > 0 )
+		{
+			DeleteVertexInstances( OrphanedVertexInstanceIDs, false );
+		}
 	}
 
-	RemovePolygonPerimeterVertices_Internal( PolygonRef, FirstVertexNumberToRemove, NumVerticesToRemove );
+	PolygonsPendingTriangulation.Add( PolygonID );
+	PolygonsPendingNewTangentBasis.Add( PolygonID );	// @todo mesheditor: and other vertex connected polygons too?
 
-	AddUndo( MakeUnique<FInsertPolygonPerimeterVerticesChange>( MoveTemp( RevertInput ) ) );
+	UE_LOG( LogEditableMesh, Verbose, TEXT( "* RemovePolygonPerimeterVertices finished" ) );
 }
 
 
-int32 UEditableMesh::FindPolygonPerimeterVertexNumberForVertex( const FPolygonRef PolygonRef, const FVertexID VertexID ) const
+int32 UEditableMesh::FindPolygonPerimeterVertexNumberForVertex( const FPolygonID PolygonID, const FVertexID VertexID ) const
 {
 	int32 FoundPolygonVertexNumber = INDEX_NONE;
 
-	const int32 PolygonPerimeterVertexCount = this->GetPolygonPerimeterVertexCount( PolygonRef );
+	const int32 PolygonPerimeterVertexCount = this->GetPolygonPerimeterVertexCount( PolygonID );
 	for( int32 PolygonVertexNumber = 0; PolygonVertexNumber < PolygonPerimeterVertexCount; ++PolygonVertexNumber )
 	{
-		if( VertexID == this->GetPolygonPerimeterVertex( PolygonRef, PolygonVertexNumber ) )
+		if( VertexID == this->GetPolygonPerimeterVertex( PolygonID, PolygonVertexNumber ) )
 		{
 			FoundPolygonVertexNumber = PolygonVertexNumber;
 			break;
@@ -5399,14 +7493,14 @@
 }
 
 
-int32 UEditableMesh::FindPolygonHoleVertexNumberForVertex( const FPolygonRef PolygonRef, const int32 HoleNumber, const FVertexID VertexID ) const
+int32 UEditableMesh::FindPolygonHoleVertexNumberForVertex( const FPolygonID PolygonID, const int32 HoleNumber, const FVertexID VertexID ) const
 {
 	int32 FoundPolygonVertexNumber = INDEX_NONE;
 
-	const int32 PolygonHoleVertexCount = this->GetPolygonHoleVertexCount( PolygonRef, HoleNumber );
+	const int32 PolygonHoleVertexCount = this->GetPolygonHoleVertexCount( PolygonID, HoleNumber );
 	for( int32 PolygonVertexNumber = 0; PolygonVertexNumber < PolygonHoleVertexCount; ++PolygonVertexNumber )
 	{
-		if( VertexID == this->GetPolygonHoleVertex( PolygonRef, HoleNumber, PolygonVertexNumber ) )
+		if( VertexID == this->GetPolygonHoleVertex( PolygonID, HoleNumber, PolygonVertexNumber ) )
 		{
 			FoundPolygonVertexNumber = PolygonVertexNumber;
 			break;
@@ -5417,12 +7511,14 @@
 }
 
 
-int32 UEditableMesh::FindPolygonPerimeterEdgeNumberForVertices( const FPolygonRef PolygonRef, const FVertexID EdgeVertexID0, const FVertexID EdgeVertexID1 ) const
+int32 UEditableMesh::FindPolygonPerimeterEdgeNumberForVertices( const FPolygonID PolygonID, const FVertexID EdgeVertexID0, const FVertexID EdgeVertexID1 ) const
 {
+	// @todo mesheditor: surely quicker just to iterate round the perimeter, checking pairs of adjacent vertices against the passed IDs (in both orders)?
+
 	int32 FoundPolygonEdgeNumber = INDEX_NONE;
 
 	static TArray<FEdgeID> EdgeIDs;
-	GetPolygonPerimeterEdges( PolygonRef, /* Out */ EdgeIDs );
+	GetPolygonPerimeterEdges( PolygonID, /* Out */ EdgeIDs );
 
 	for( int32 PolygonEdgeNumber = 0; PolygonEdgeNumber < EdgeIDs.Num(); ++PolygonEdgeNumber )
 	{
@@ -5443,12 +7539,14 @@
 }
 
 
-int32 UEditableMesh::FindPolygonHoleEdgeNumberForVertices( const FPolygonRef PolygonRef, const int32 HoleNumber, const FVertexID EdgeVertexID0, const FVertexID EdgeVertexID1 ) const
+int32 UEditableMesh::FindPolygonHoleEdgeNumberForVertices( const FPolygonID PolygonID, const int32 HoleNumber, const FVertexID EdgeVertexID0, const FVertexID EdgeVertexID1 ) const
 {
+	// @todo mesheditor: surely quicker just to iterate round the perimeter, checking pairs of adjacent vertices against the passed IDs (in both orders)?
+
 	int32 FoundPolygonEdgeNumber = INDEX_NONE;
 
 	static TArray<FEdgeID> EdgeIDs;
-	GetPolygonHoleEdges( PolygonRef, HoleNumber, /* Out */ EdgeIDs );
+	GetPolygonHoleEdges( PolygonID, HoleNumber, /* Out */ EdgeIDs );
 
 	for( int32 PolygonEdgeNumber = 0; PolygonEdgeNumber < EdgeIDs.Num(); ++PolygonEdgeNumber )
 	{
@@ -5469,33 +7567,34 @@
 }
 
 
-void UEditableMesh::FlipPolygons( const TArray<FPolygonRef>& PolygonRefs )
+void UEditableMesh::FlipPolygons( const TArray<FPolygonID>& PolygonIDs )
 {
 	static TArray<FPolygonToCreate> PolygonsToCreate;
 	PolygonsToCreate.Reset();
-	for( int32 PolygonNumber = 0; PolygonNumber < PolygonRefs.Num(); ++PolygonNumber )
+	for( int32 PolygonNumber = 0; PolygonNumber < PolygonIDs.Num(); ++PolygonNumber )
 	{
-		const FPolygonRef OriginalPolygonRef = PolygonRefs[ PolygonNumber ];
+		const FPolygonID OriginalPolygonID = PolygonIDs[ PolygonNumber ];
+		const FPolygonGroupID OriginalPolygonGroupID = GetGroupForPolygon( OriginalPolygonID );
 
 		FPolygonToCreate& PolygonToCreate = *new( PolygonsToCreate ) FPolygonToCreate();
 
-		PolygonToCreate.SectionID = OriginalPolygonRef.SectionID;
+		PolygonToCreate.PolygonGroupID = OriginalPolygonGroupID;
 		
 		// Keep the original polygon ID.  No reason not to.
-		PolygonToCreate.OriginalPolygonID = OriginalPolygonRef.PolygonID;
+		PolygonToCreate.OriginalPolygonID = OriginalPolygonID;
 
 		// Iterate backwards to add the vertices for the polygon, because we're flipping it over.
-		const int32 PerimeterVertexCount = GetPolygonPerimeterVertexCount( OriginalPolygonRef );
+		const int32 PerimeterVertexCount = GetPolygonPerimeterVertexCount( OriginalPolygonID );
 		PolygonToCreate.PerimeterVertices.Reserve( PerimeterVertexCount );
 		for( int32 VertexNumber = PerimeterVertexCount - 1; VertexNumber >= 0; --VertexNumber )
 		{
 			FVertexAndAttributes& PerimeterVertex = *new( PolygonToCreate.PerimeterVertices ) FVertexAndAttributes();
 
-			PerimeterVertex.VertexID = GetPolygonPerimeterVertex( OriginalPolygonRef, VertexNumber );
+			PerimeterVertex.VertexID = GetPolygonPerimeterVertex( OriginalPolygonID, VertexNumber );
 
 			// Save off information about the polygons, so we can re-create them with reverse winding.
 			TArray<FMeshElementAttributeData>& PerimeterVertexAttributeList = PerimeterVertex.PolygonVertexAttributes.Attributes;
-			for( const FName AttributeName : UEditableMesh::GetValidPolygonVertexAttributes() )
+			for( const FName AttributeName : UEditableMesh::GetValidVertexInstanceAttributes() )
 			{
 				const int32 MaxAttributeIndex = GetMaxAttributeIndex( AttributeName );
 				for( int32 AttributeIndex = 0; AttributeIndex < MaxAttributeIndex; ++AttributeIndex )
@@ -5503,7 +7602,7 @@
 					FMeshElementAttributeData& PolygonVertexAttribute = *new( PerimeterVertexAttributeList ) FMeshElementAttributeData(
 						AttributeName,
 						AttributeIndex,
-						GetPolygonPerimeterVertexAttribute( OriginalPolygonRef, VertexNumber, AttributeName, AttributeIndex ) );
+						GetPolygonPerimeterVertexAttribute( OriginalPolygonID, VertexNumber, AttributeName, AttributeIndex ) );
 				}
 			}
 		}
@@ -5515,78 +7614,60 @@
 	{
 		const bool bDeleteOrphanedEdges = false;
 		const bool bDeleteOrphanedVertices = false;
-		const bool bDeleteEmptySections = false;
-		DeletePolygons( PolygonRefs, bDeleteOrphanedEdges, bDeleteOrphanedVertices, bDeleteEmptySections );
+		const bool bDeleteOrphanedVertexInstances = false;
+		const bool bDeleteEmptyPolygonGroups = false;
+		DeletePolygons( PolygonIDs, bDeleteOrphanedEdges, bDeleteOrphanedVertices, bDeleteOrphanedVertexInstances, bDeleteEmptyPolygonGroups );
 	}
 
 	// Create the new polygons.  They're just like the polygons we deleted, except with reversed winding.
 	{
-		static TArray<FPolygonRef> NewPolygonRefs;
+		static TArray<FPolygonID> NewPolygonIDs;
 		static TArray<FEdgeID> NewEdgeIDs;
-		CreatePolygons( PolygonsToCreate, NewPolygonRefs, NewEdgeIDs );
-
-		// Generate fresh normals
-		GenerateNormalsAndTangentsForPolygonsAndAdjacents( NewPolygonRefs );
+		CreatePolygons( PolygonsToCreate, NewPolygonIDs, NewEdgeIDs );
 	}
 }
 
 
-void UEditableMesh::TriangulatePolygons( const TArray<FPolygonRef>& PolygonRefs, TArray<FPolygonRef>& OutNewTrianglePolygons )
+void UEditableMesh::TriangulatePolygons( const TArray<FPolygonID>& PolygonIDs, TArray<FPolygonID>& OutNewTrianglePolygons )
 {
 	OutNewTrianglePolygons.Reset();
 
 	static TArray<FPolygonToCreate> PolygonsToCreate;
 	PolygonsToCreate.Reset();
 
-	static TArray<FPolygonRef> PolygonsToDelete;
+	static TArray<FPolygonID> PolygonsToDelete;
 	PolygonsToDelete.Reset();
 
-	for( const FPolygonRef PolygonRef : PolygonRefs )
+	for( const FPolygonID PolygonID : PolygonIDs )
 	{
 		// Skip right over polygons with fewer than four vertices
-		const int32 PerimeterVertexCount = GetPolygonPerimeterVertexCount( PolygonRef );
+		const int32 PerimeterVertexCount = GetPolygonPerimeterVertexCount( PolygonID );
 		if( PerimeterVertexCount > 3 )
 		{
 			// We'll be replacing this polygon with it's triangulated counterpart polygons
-			PolygonsToDelete.Add( PolygonRef );
+			PolygonsToDelete.Add( PolygonID );
 
 			// Figure out the triangulation for this polygon
-			static TArray<int32> PerimeterVertexNumbersForTriangles;
-			PerimeterVertexNumbersForTriangles.Reset();
-			ComputePolygonTriangulation( PolygonRef, /* Out */ PerimeterVertexNumbersForTriangles );
+			static TArray<FMeshTriangle> Triangles;
+			ComputePolygonTriangulation( PolygonID, /* Out */ Triangles );
 
 			// Build polygons for each of the triangles that made up the original
 			{
-				check( ( PerimeterVertexNumbersForTriangles.Num() % 3 ) == 0 );	// Expecting three indices per triangle
-				const int32 TriangleCount = PerimeterVertexNumbersForTriangles.Num() / 3;
-				for( int32 TriangleIndex = 0; TriangleIndex < TriangleCount; ++TriangleIndex )
+				for( const FMeshTriangle& Triangle : Triangles )
 				{
-					const int32 FirstTrianglesVertexIDIndex = TriangleIndex * 3;
+					PolygonsToCreate.Emplace();
+					FPolygonToCreate& PolygonToCreate = PolygonsToCreate.Last();
+					
+					PolygonToCreate.PolygonGroupID = GetGroupForPolygon( PolygonID );
+					PolygonToCreate.PolygonEdgeHardness = EPolygonEdgeHardness::NewEdgesSoft;
 
-					FPolygonToCreate& PolygonToCreate = *new( PolygonsToCreate ) FPolygonToCreate();
-					PolygonToCreate.SectionID = PolygonRef.SectionID;
-
 					for( int32 TriangleVertexNumber = 0; TriangleVertexNumber < 3; ++TriangleVertexNumber )
 					{
-						const int32 PerimeterVertexNumber = PerimeterVertexNumbersForTriangles[ FirstTrianglesVertexIDIndex + TriangleVertexNumber ];
-						const FVertexID TriangleVertexID = GetPolygonPerimeterVertex( PolygonRef, PerimeterVertexNumber );
+						PolygonToCreate.PerimeterVertices.Emplace();
+						FVertexAndAttributes& PerimeterVertex = PolygonToCreate.PerimeterVertices.Last();
 
-						FVertexAndAttributes& PerimeterVertex = *new( PolygonToCreate.PerimeterVertices ) FVertexAndAttributes();
-						PerimeterVertex.VertexID = TriangleVertexID;
-
-						// Save off information about the polygons, so we can re-create them with reverse winding.
-						TArray<FMeshElementAttributeData>& PerimeterVertexAttributeList = PerimeterVertex.PolygonVertexAttributes.Attributes;
-						for( const FName AttributeName : UEditableMesh::GetValidPolygonVertexAttributes() )
-						{
-							const int32 MaxAttributeIndex = GetMaxAttributeIndex( AttributeName );
-							for( int32 AttributeIndex = 0; AttributeIndex < MaxAttributeIndex; ++AttributeIndex )
-							{
-								FMeshElementAttributeData& PolygonVertexAttribute = *new( PerimeterVertexAttributeList ) FMeshElementAttributeData(
-									AttributeName,
-									AttributeIndex,
-									GetPolygonPerimeterVertexAttribute( PolygonRef, PerimeterVertexNumber, AttributeName, AttributeIndex ) );
-							}
-						}
+						PerimeterVertex.VertexInstanceID = Triangle.VertexInstanceIDs[ TriangleVertexNumber ];
+						PerimeterVertex.VertexID = FVertexID::Invalid;
 					}
 				}
 			}
@@ -5598,8 +7679,9 @@
 	{
 		const bool bDeleteOrphanEdges = false;
 		const bool bDeleteOrphanVertices = false;
-		const bool bDeleteEmptySections = false;
-		DeletePolygons( PolygonsToDelete, bDeleteOrphanEdges, bDeleteOrphanVertices, bDeleteEmptySections );
+		const bool bDeleteOrphanVertexInstances = false;
+		const bool bDeleteEmptyPolygonGroups = false;
+		DeletePolygons( PolygonsToDelete, bDeleteOrphanEdges, bDeleteOrphanVertices, bDeleteOrphanVertexInstances, bDeleteEmptyPolygonGroups );
 	}
 
 	// Create the new polygons.  One for each triangle.  Note that new edges will be created here too on the inside of
@@ -5607,56 +7689,45 @@
 	{
 		static TArray<FEdgeID> NewEdgeIDs;
 		CreatePolygons( PolygonsToCreate, /* Out */ OutNewTrianglePolygons, /* Out */ NewEdgeIDs );
-
-		// @todo mesheditor: The the new internal edges of the polygon (NewEdgeIDs) get hard or soft edges by default?\
-
-		// Generate fresh normals
-		GenerateNormalsAndTangentsForPolygonsAndAdjacents( OutNewTrianglePolygons );
 	}
 }
 
 
-FSectionID UEditableMesh::CreateSection( const FSectionToCreate& SectionToCreate )
+void UEditableMesh::AssignMaterialToPolygons( const TArray<FPolygonID>& PolygonIDs, UMaterialInterface* Material, TArray<FPolygonID>& NewPolygonIDs )
 {
-	return CreateSection_Internal( SectionToCreate );
-}
+	// @todo mesheditor: create more generic method to assign a polygon to a new polygon group
 
+	// We need to move polygons from one polygon group to another. First find the new polygon group id to move them to.
+	const bool bCreateNewPolygonGroupIfNotFound = true;
+	const FPolygonGroupID NewPolygonGroupID = GetPolygonGroupIDFromMaterial( Material, bCreateNewPolygonGroupIfNotFound );
+	check( NewPolygonGroupID != FPolygonGroupID::Invalid );
 
-void UEditableMesh::DeleteSection( const FSectionID SectionID )
-{
-	DeleteSection_Internal( SectionID );
-}
-
-
-void UEditableMesh::AssignMaterialToPolygons( const TArray<FPolygonRef>& PolygonRefs, UMaterialInterface* Material, TArray<FPolygonRef>& NewPolygonRefs )
-{
-	// We need to move polygons from one section to another. First find the new section id to move them to.
-	const bool bCreateNewSectionIfNotFound = true;
-	const FSectionID NewSectionID = GetSectionIDFromMaterial_Internal( Material, bCreateNewSectionIfNotFound );
-	check( NewSectionID != FSectionID::Invalid );
-
-	// Make an array of FPolygonToCreate, based on the ones we wish to move. Everything is the same, other than the SectionID
+	// Make an array of FPolygonToCreate, based on the ones we wish to move. Everything is the same, other than the PolygonGroupID
 	static TArray<FPolygonToCreate> PolygonsToCreate;
-	PolygonsToCreate.Reset( PolygonRefs.Num() );
-	for( int32 PolygonNumber = 0; PolygonNumber < PolygonRefs.Num(); ++PolygonNumber )
+	PolygonsToCreate.Reset( PolygonIDs.Num() );
+
+	for( int32 PolygonNumber = 0; PolygonNumber < PolygonIDs.Num(); ++PolygonNumber )
 	{
-		const FPolygonRef OriginalPolygonRef = PolygonRefs[ PolygonNumber ];
+		const FPolygonID OriginalPolygonID = PolygonIDs[ PolygonNumber ];
+		const FMeshPolygon& Polygon = Polygons[ OriginalPolygonID.GetValue() ];
 
 		FPolygonToCreate& PolygonToCreate = *new( PolygonsToCreate ) FPolygonToCreate();
 
-		PolygonToCreate.SectionID = NewSectionID;
+		PolygonToCreate.PolygonGroupID = NewPolygonGroupID;
 		PolygonToCreate.OriginalPolygonID = FPolygonID::Invalid;
 
-		const int32 PerimeterVertexCount = GetPolygonPerimeterVertexCount( OriginalPolygonRef );
+		const int32 PerimeterVertexCount = GetPolygonPerimeterVertexCount( OriginalPolygonID );
 		PolygonToCreate.PerimeterVertices.Reserve( PerimeterVertexCount );
+
 		for( int32 VertexNumber = 0; VertexNumber < PerimeterVertexCount; ++VertexNumber )
 		{
 			FVertexAndAttributes& PerimeterVertex = *new( PolygonToCreate.PerimeterVertices ) FVertexAndAttributes();
 
-			PerimeterVertex.VertexID = GetPolygonPerimeterVertex( OriginalPolygonRef, VertexNumber );
+			PerimeterVertex.VertexID = GetPolygonPerimeterVertex( OriginalPolygonID, VertexNumber );
 
 			TArray<FMeshElementAttributeData>& PerimeterVertexAttributeList = PerimeterVertex.PolygonVertexAttributes.Attributes;
-			for( const FName AttributeName : UEditableMesh::GetValidPolygonVertexAttributes() )
+
+			for( const FName AttributeName : UEditableMesh::GetValidVertexInstanceAttributes() )
 			{
 				const int32 MaxAttributeIndex = GetMaxAttributeIndex( AttributeName );
 				for( int32 AttributeIndex = 0; AttributeIndex < MaxAttributeIndex; ++AttributeIndex )
@@ -5664,7 +7735,7 @@
 					FMeshElementAttributeData& PolygonVertexAttribute = *new( PerimeterVertexAttributeList ) FMeshElementAttributeData(
 						AttributeName,
 						AttributeIndex,
-						GetPolygonPerimeterVertexAttribute( OriginalPolygonRef, VertexNumber, AttributeName, AttributeIndex ) );
+						GetPolygonPerimeterVertexAttribute( OriginalPolygonID, VertexNumber, AttributeName, AttributeIndex ) );
 				}
 			}
 		}
@@ -5672,22 +7743,57 @@
 		// @todo mesheditor: Handle holes?
 	}
 
-	// Create the new polygons in the new section.
+	// Create the new polygons in the new PolygonGroup.
 	{
 		static TArray<FEdgeID> NewEdgeIDs;
-		CreatePolygons( PolygonsToCreate, NewPolygonRefs, NewEdgeIDs );
+		CreatePolygons( PolygonsToCreate, NewPolygonIDs, NewEdgeIDs );
 	}
 
-	// Delete the polygons from the old section
+	// Delete the polygons from the old PolygonGroup
 	{
 		const bool bDeleteOrphanedEdges = false;
 		const bool bDeleteOrphanedVertices = false;
-		const bool bDeleteEmptySections = true;
-		DeletePolygons( PolygonRefs, bDeleteOrphanedEdges, bDeleteOrphanedVertices, bDeleteEmptySections );
+		const bool bDeleteOrphanedVertexInstances = false;
+		const bool bDeleteEmptyPolygonGroups = true;
+		DeletePolygons( PolygonIDs, bDeleteOrphanedEdges, bDeleteOrphanedVertices, bDeleteOrphanedVertexInstances, bDeleteEmptyPolygonGroups );
 	}
 }
 
 
+FPolygonGroupID UEditableMesh::GetPolygonGroupIDFromMaterial( UMaterialInterface* Material, bool bCreateNewSectionIfNotFound )
+{
+	// Iterate through the sections sparse array looking for an entry whose material index matches.
+	for( TSparseArray<FMeshPolygonGroup>::TConstIterator It( PolygonGroups ); It; ++It )
+	{
+		const FMeshPolygonGroup& PolygonGroup = *It;
+
+		if( PolygonGroup.Material == Material )
+		{
+			return FPolygonGroupID( It.GetIndex() );
+		}
+	}
+
+	// If we got here, the material index does not yet have a matching section.
+	if( bCreateNewSectionIfNotFound )
+	{
+		static TArray<FPolygonGroupToCreate> PolygonGroupsToCreate;
+		PolygonGroupsToCreate.Reset( 1 );
+		PolygonGroupsToCreate.Emplace();
+		FPolygonGroupToCreate& PolygonGroupToCreate = PolygonGroupsToCreate[ 0 ];
+
+		PolygonGroupToCreate.Material = Material;
+		PolygonGroupToCreate.bEnableCollision = true;
+		PolygonGroupToCreate.bCastShadow = true;
+
+		static TArray<FPolygonGroupID> PolygonGroupIDs;
+		CreatePolygonGroups( PolygonGroupsToCreate, PolygonGroupIDs );
+		return PolygonGroupIDs[ 0 ];
+	}
+
+	return FPolygonGroupID::Invalid;
+}
+
+
 void UEditableMesh::WeldVertices( const TArray<FVertexID>& VertexIDsToWeld, FVertexID& OutNewVertexID )
 {
 	OutNewVertexID = FVertexID::Invalid;
@@ -5757,18 +7863,18 @@
 	};
 
 	// Build a list of all polygons which contain at least one of the vertices to be welded
-	static TArray<FPolygonRef> AllConnectedPolygonRefs;
+	static TArray<FPolygonID> AllConnectedPolygonIDs;
 	{
-		AllConnectedPolygonRefs.Reset();
+		AllConnectedPolygonIDs.Reset();
 
 		for( FVertexID VertexID : VertexIDsToWeld )
 		{
-			static TArray<FPolygonRef> ConnectedPolygonRefs;
-			GetVertexConnectedPolygons( VertexID, ConnectedPolygonRefs );
+			static TArray<FPolygonID> ConnectedPolygonIDs;
+			GetVertexConnectedPolygons( VertexID, ConnectedPolygonIDs );
 
-			for( FPolygonRef PolygonRef : ConnectedPolygonRefs )
+			for( FPolygonID PolygonID : ConnectedPolygonIDs )
 			{
-				AllConnectedPolygonRefs.AddUnique( PolygonRef );
+				AllConnectedPolygonIDs.AddUnique( PolygonID );
 			}
 		}
 	}
@@ -5776,10 +7882,10 @@
 	// Check whether the operation is valid. We can't weld vertices if there are any polygons which have non-contiguous
 	// vertices on the perimeter which are marked to be welded
 	bool bNeedToCreateWeldedVertex = false;
-	for( FPolygonRef ConnectedPolygonRef : AllConnectedPolygonRefs )
+	for( FPolygonID ConnectedPolygonID : AllConnectedPolygonIDs )
 	{
 		static TArray<FVertexID> PolygonVertexIDs;
-		GetPolygonPerimeterVertices( ConnectedPolygonRef, PolygonVertexIDs );
+		GetPolygonPerimeterVertices( ConnectedPolygonID, PolygonVertexIDs );
 
 		auto VertexRangeToWeld = GetPerimeterVertexRangeToWeld( PolygonVertexIDs, VertexIDsToWeld );
 		const bool bIsValid = VertexRangeToWeld.Get<0>();
@@ -5829,21 +7935,21 @@
 
 
 	static TArray<FPolygonToCreate> PolygonsToCreate;
-	PolygonsToCreate.Reset( AllConnectedPolygonRefs.Num() );
+	PolygonsToCreate.Reset( AllConnectedPolygonIDs.Num() );
 
 	static TArray<FAttributesForEdge> AttributesForEdges;
 	AttributesForEdges.Reset();
 
 	// Now for each polygon, merge runs of consecutive vertices
-	for( FPolygonRef ConnectedPolygonRef : AllConnectedPolygonRefs )
+	for( FPolygonID ConnectedPolygonID : AllConnectedPolygonIDs )
 	{
-		const int32 NumPerimeterVertices = GetPolygonPerimeterVertexCount( ConnectedPolygonRef );
+		const int32 NumPerimeterVertices = GetPolygonPerimeterVertexCount( ConnectedPolygonID );
 
 		// Get perimeter vertices and edges for this polygon
 		static TArray<FVertexID> PolygonVertexIDs;
 		static TArray<FEdgeID> PolygonEdgeIDs;
-		GetPolygonPerimeterVertices( ConnectedPolygonRef, PolygonVertexIDs );
-		GetPolygonPerimeterEdges( ConnectedPolygonRef, PolygonEdgeIDs );
+		GetPolygonPerimeterVertices( ConnectedPolygonID, PolygonVertexIDs );
+		GetPolygonPerimeterEdges( ConnectedPolygonID, PolygonEdgeIDs );
 
 		// Get the index range of perimeter vertices to be welded.
 		// This should definitely be valid, as any invalid welded poly will have caused early exit, above.
@@ -5862,7 +7968,8 @@
 		// Prepare to create a new polygon
 		PolygonsToCreate.Emplace();
 		FPolygonToCreate& PolygonToCreate = PolygonsToCreate.Last();
-		PolygonToCreate.SectionID = ConnectedPolygonRef.SectionID;
 
+		PolygonToCreate.PolygonGroupID = GetGroupForPolygon( ConnectedPolygonID );
+
 		// Iterate through perimeter vertices starting at index 0.
 		// We skip through the run of welded vertices, replacing them with a single welded vertex.
@@ -5906,14 +8013,14 @@
 				VertexAndAttributes.PolygonVertexAttributes.Attributes.Add( FMeshElementAttributeData(
 					UEditableMeshAttribute::VertexTextureCoordinate(),
 					TextureCoordinateIndex,
-					GetPolygonPerimeterVertexAttribute( ConnectedPolygonRef, Index, UEditableMeshAttribute::VertexTextureCoordinate(), TextureCoordinateIndex ) )
+					GetPolygonPerimeterVertexAttribute( ConnectedPolygonID, Index, UEditableMeshAttribute::VertexTextureCoordinate(), TextureCoordinateIndex ) )
 				);
 			}
 
 			VertexAndAttributes.PolygonVertexAttributes.Attributes.Add( FMeshElementAttributeData(
 				UEditableMeshAttribute::VertexColor(),
 				0,
-				GetPolygonPerimeterVertexAttribute( ConnectedPolygonRef, Index, UEditableMeshAttribute::VertexColor(), 0 ) )
+				GetPolygonPerimeterVertexAttribute( ConnectedPolygonID, Index, UEditableMeshAttribute::VertexColor(), 0 ) )
 			);
 
 			// Prepare to assign the old edge's attributes to the new edge.
@@ -5935,22 +8042,22 @@
 	}
 
 	// Create polygons
-	static TArray<FPolygonRef> NewPolygonRefs;
+	static TArray<FPolygonID> NewPolygonIDs;
 	{
 		static TArray<FEdgeID> NewEdgeIDs;
-		CreatePolygons( PolygonsToCreate, NewPolygonRefs, NewEdgeIDs );
+		CreatePolygons( PolygonsToCreate, NewPolygonIDs, NewEdgeIDs );
 	}
 
 	// Set new edge attributes.
-	// Now that we have a list of newly created polygon refs, we need to go through the attributes for edge list, filling in the new Edge ID.
-	// This relies on the fact that the NewPolygonRefs array lists the polygons in the same order as they were defined in PolygonsToCreate,
+	// Now that we have a list of newly created polygon IDs, we need to go through the attributes for edge list, filling in the new Edge ID.
+	// This relies on the fact that the NewPolygonIDs array lists the polygons in the same order as they were defined in PolygonsToCreate,
 	// and that the edges are strictly ordered from perimeter vertex 0.
 	{
 		int32 AttributesForEdgeIndex = 0;
-		for( FPolygonRef NewPolygonRef : NewPolygonRefs )
+		for( FPolygonID NewPolygonID : NewPolygonIDs )
 		{
 			static TArray<FEdgeID> NewPolygonEdgeIDs;
-			GetPolygonPerimeterEdges( NewPolygonRef, NewPolygonEdgeIDs );
+			GetPolygonPerimeterEdges( NewPolygonID, NewPolygonEdgeIDs );
 
 			for( FEdgeID NewPolygonEdgeID : NewPolygonEdgeIDs )
 			{
@@ -5966,17 +8073,16 @@
 	{
 		const bool bDeleteOrphanedEdges = true;
 		const bool bDeleteOrphanedVertices = true;
-		const bool bDeleteEmptySections = false;
-		DeletePolygons( AllConnectedPolygonRefs, bDeleteOrphanedEdges, bDeleteOrphanedVertices, bDeleteEmptySections );
+		const bool bDeleteOrphanedVertexInstances = true;
+		const bool bDeleteEmptyPolygonGroups = false;
+		DeletePolygons( AllConnectedPolygonIDs, bDeleteOrphanedEdges, bDeleteOrphanedVertices, bDeleteOrphanedVertexInstances, bDeleteEmptyPolygonGroups );
 	}
-
-	GenerateNormalsAndTangentsForPolygonsAndAdjacents( NewPolygonRefs );
 }
 
 
-void UEditableMesh::TessellatePolygons( const TArray<FPolygonRef>& PolygonRefs, const ETriangleTessellationMode TriangleTessellationMode, TArray<FPolygonRef>& OutNewPolygonRefs )
+void UEditableMesh::TessellatePolygons( const TArray<FPolygonID>& PolygonIDs, const ETriangleTessellationMode TriangleTessellationMode, TArray<FPolygonID>& OutNewPolygonIDs )
 {
-	OutNewPolygonRefs.Reset();
+	OutNewPolygonIDs.Reset();
 
 	//
 	// Simple tessellation algorithm:
@@ -5997,13 +8103,13 @@
 	{
 		static TArray<FVertexToCreate> VerticesToCreate;
 		VerticesToCreate.Reset();
-		for( const FPolygonRef PolygonRef : PolygonRefs )
+		for( const FPolygonID PolygonID : PolygonIDs )
 		{
-			const int32 PerimeterEdgeCount = GetPolygonPerimeterEdgeCount( PolygonRef );
+			const int32 PerimeterEdgeCount = GetPolygonPerimeterEdgeCount( PolygonID );
 			if( TriangleTessellationMode == ETriangleTessellationMode::ThreeTriangles || PerimeterEdgeCount > 3 )
 			{
 				// Find the center of this polygon
-				const FVector PolygonCenter = ComputePolygonCenter( PolygonRef );
+				const FVector PolygonCenter = ComputePolygonCenter( PolygonID );
 
 				FVertexToCreate& VertexToCreate = *new( VerticesToCreate ) FVertexToCreate();
 				VertexToCreate.VertexAttributes.Attributes.Add( FMeshElementAttributeData( UEditableMeshAttribute::VertexPosition(), 0, PolygonCenter ) );
@@ -6020,12 +8126,12 @@
 		static TSet<FEdgeID> EdgesToSplit;
 		EdgesToSplit.Reset();
 
-		for( int32 PolygonNumber = 0; PolygonNumber < PolygonRefs.Num(); ++PolygonNumber )
+		for( int32 PolygonNumber = 0; PolygonNumber < PolygonIDs.Num(); ++PolygonNumber )
 		{
-			const FPolygonRef PolygonRef = PolygonRefs[ PolygonNumber ];
+			const FPolygonID PolygonID = PolygonIDs[ PolygonNumber ];
 
 			static TArray<FEdgeID> PerimeterEdgeIDs;
-			GetPolygonPerimeterEdges( PolygonRef, /* Out */ PerimeterEdgeIDs );
+			GetPolygonPerimeterEdges( PolygonID, /* Out */ PerimeterEdgeIDs );
 			
 			if( TriangleTessellationMode == ETriangleTessellationMode::FourTriangles || PerimeterEdgeIDs.Num() > 3 )
 			{
@@ -6055,11 +8161,12 @@
 	PolygonsToCreate.Reset();
 
 	int32 PolygonWithNewCenterVertexNumber = 0;
-	for( int32 PolygonNumber = 0; PolygonNumber < PolygonRefs.Num(); ++PolygonNumber )
+	for( int32 PolygonNumber = 0; PolygonNumber < PolygonIDs.Num(); ++PolygonNumber )
 	{
-		const FPolygonRef PolygonRef = PolygonRefs[ PolygonNumber ];
+		const FPolygonID PolygonID = PolygonIDs[ PolygonNumber ];
+		const FPolygonGroupID PolygonGroupID = GetGroupForPolygon( PolygonID );
 
-		const int32 PerimeterEdgeCount = GetPolygonPerimeterEdgeCount( PolygonRef );
+		const int32 PerimeterEdgeCount = GetPolygonPerimeterEdgeCount( PolygonID );
 
 		FVertexID PolygonCenterVertexID = FVertexID::Invalid;
 		if( TriangleTessellationMode == ETriangleTessellationMode::ThreeTriangles || PerimeterEdgeCount > 6 )
@@ -6072,7 +8179,7 @@
 		if( PerimeterEdgeCount > 6 )
 		{
 			static TArray<FVertexID> PerimeterVertexIDs;
-			GetPolygonPerimeterVertices( PolygonRef, /* Out */ PerimeterVertexIDs );
+			GetPolygonPerimeterVertices( PolygonID, /* Out */ PerimeterVertexIDs );
 
 			const int32 PerimeterVertexCount = PerimeterEdgeCount;
 			const int32 OriginalPerimeterEdgeCount = PerimeterEdgeCount / 2;
@@ -6083,7 +8190,7 @@
 				const int32 NextVertexNumber = ( CurrentVertexNumber + 1 ) % PerimeterVertexCount;
 
 				FPolygonToCreate& PolygonToCreate = *new( PolygonsToCreate ) FPolygonToCreate();
-				PolygonToCreate.SectionID = PolygonRef.SectionID;
+				PolygonToCreate.PolygonGroupID = PolygonGroupID;
 
 				for( int32 QuadVertexNumber = 0; QuadVertexNumber < 4; ++QuadVertexNumber )
 				{
@@ -6121,26 +8228,26 @@
 						{
 							const FVector CenterVertexPosition = GetVertexAttribute( PolygonCenterVertexID, UEditableMeshAttribute::VertexPosition(), 0 );
 
-							static TArray<int32> PerimeterVertexIndices;
+							FMeshTriangle Triangle;
 							FVector TriangleVertexWeights;
-							if( ComputeBarycentricWeightForPointOnPolygon( PolygonRef, CenterVertexPosition, /* Out */ PerimeterVertexIndices, /* Out */ TriangleVertexWeights ) )
+							if( ComputeBarycentricWeightForPointOnPolygon( PolygonID, CenterVertexPosition, /* Out */ Triangle, /* Out */ TriangleVertexWeights ) )
 							{
 								for( int32 TextureCoordinateIndex = 0; TextureCoordinateIndex < TextureCoordinateCount; ++TextureCoordinateIndex )
 								{
 									VertexAndAttributes.PolygonVertexAttributes.Attributes.Add( FMeshElementAttributeData(
 										UEditableMeshAttribute::VertexTextureCoordinate(),
 										TextureCoordinateIndex,
-										TriangleVertexWeights.X * GetPolygonPerimeterVertexAttribute( PolygonRef, PerimeterVertexIndices[ 0 ], UEditableMeshAttribute::VertexTextureCoordinate(), TextureCoordinateIndex ) +
-										TriangleVertexWeights.Y * GetPolygonPerimeterVertexAttribute( PolygonRef, PerimeterVertexIndices[ 1 ], UEditableMeshAttribute::VertexTextureCoordinate(), TextureCoordinateIndex ) +
-										TriangleVertexWeights.Z * GetPolygonPerimeterVertexAttribute( PolygonRef, PerimeterVertexIndices[ 2 ], UEditableMeshAttribute::VertexTextureCoordinate(), TextureCoordinateIndex ) ) );
+										TriangleVertexWeights.X * GetVertexInstanceAttribute( Triangle.VertexInstanceIDs[ 0 ], UEditableMeshAttribute::VertexTextureCoordinate(), TextureCoordinateIndex ) +
+										TriangleVertexWeights.Y * GetVertexInstanceAttribute( Triangle.VertexInstanceIDs[ 1 ], UEditableMeshAttribute::VertexTextureCoordinate(), TextureCoordinateIndex ) +
+										TriangleVertexWeights.Z * GetVertexInstanceAttribute( Triangle.VertexInstanceIDs[ 2 ], UEditableMeshAttribute::VertexTextureCoordinate(), TextureCoordinateIndex ) ) );
 								}
 
 								VertexAndAttributes.PolygonVertexAttributes.Attributes.Add( FMeshElementAttributeData(
 									UEditableMeshAttribute::VertexColor(),
 									0,
-									TriangleVertexWeights.X * GetPolygonPerimeterVertexAttribute( PolygonRef, PerimeterVertexIndices[ 0 ], UEditableMeshAttribute::VertexColor(), 0 ) +
-									TriangleVertexWeights.Y * GetPolygonPerimeterVertexAttribute( PolygonRef, PerimeterVertexIndices[ 1 ], UEditableMeshAttribute::VertexColor(), 0 ) +
-									TriangleVertexWeights.Z * GetPolygonPerimeterVertexAttribute( PolygonRef, PerimeterVertexIndices[ 2 ], UEditableMeshAttribute::VertexColor(), 0 ) ) );
+									TriangleVertexWeights.X * GetVertexInstanceAttribute( Triangle.VertexInstanceIDs[ 0 ], UEditableMeshAttribute::VertexColor(), 0 ) +
+									TriangleVertexWeights.Y * GetVertexInstanceAttribute( Triangle.VertexInstanceIDs[ 1 ], UEditableMeshAttribute::VertexColor(), 0 ) +
+									TriangleVertexWeights.Z * GetVertexInstanceAttribute( Triangle.VertexInstanceIDs[ 2 ], UEditableMeshAttribute::VertexColor(), 0 ) ) );
 							}
 						}
 					}
@@ -6155,13 +8262,13 @@
 								VertexAndAttributes.PolygonVertexAttributes.Attributes.Add( FMeshElementAttributeData(
 									UEditableMeshAttribute::VertexTextureCoordinate(),
 									TextureCoordinateIndex,
-									GetPolygonPerimeterVertexAttribute( PolygonRef, PerimeterVertexNumber, UEditableMeshAttribute::VertexTextureCoordinate(), TextureCoordinateIndex ) ) );
+									GetPolygonPerimeterVertexAttribute( PolygonID, PerimeterVertexNumber, UEditableMeshAttribute::VertexTextureCoordinate(), TextureCoordinateIndex ) ) );
 							}
 
 							VertexAndAttributes.PolygonVertexAttributes.Attributes.Add( FMeshElementAttributeData(
 								UEditableMeshAttribute::VertexColor(),
 								0,
-								GetPolygonPerimeterVertexAttribute( PolygonRef, PerimeterVertexNumber, UEditableMeshAttribute::VertexColor(), 0 ) ) );
+								GetPolygonPerimeterVertexAttribute( PolygonID, PerimeterVertexNumber, UEditableMeshAttribute::VertexColor(), 0 ) ) );
 						}
 					}
 				}
@@ -6175,13 +8282,13 @@
 				for( int32 PerimeterEdgeNumber = 0; PerimeterEdgeNumber < 3; ++PerimeterEdgeNumber )
 				{
 					bool bIsEdgeWindingReversedForPolygon;
-					const FEdgeID EdgeID = GetPolygonPerimeterEdge( PolygonRef, PerimeterEdgeNumber, /* Out */ bIsEdgeWindingReversedForPolygon );
+					const FEdgeID EdgeID = GetPolygonPerimeterEdge( PolygonID, PerimeterEdgeNumber, /* Out */ bIsEdgeWindingReversedForPolygon );
 
 					FVertexID EdgeVertexID0, EdgeVertexID1;
 					GetEdgeVertices( EdgeID, /* Out */ EdgeVertexID0, /* Out */ EdgeVertexID1 );
 
 					FPolygonToCreate& PolygonToCreate = *new( PolygonsToCreate ) FPolygonToCreate();
-					PolygonToCreate.SectionID = PolygonRef.SectionID;
+					PolygonToCreate.PolygonGroupID = PolygonGroupID;
 
 					for( int32 TriangleVertexNumber = 0; TriangleVertexNumber < 3; ++TriangleVertexNumber )
 					{
@@ -6210,26 +8317,26 @@
 							{
 								const FVector CenterVertexPosition = GetVertexAttribute( PolygonCenterVertexID, UEditableMeshAttribute::VertexPosition(), 0 );
 
-								static TArray<int32> PerimeterVertexIndices;
+								FMeshTriangle Triangle;
 								FVector TriangleVertexWeights;
-								if( ComputeBarycentricWeightForPointOnPolygon( PolygonRef, CenterVertexPosition, /* Out */ PerimeterVertexIndices, /* Out */ TriangleVertexWeights ) )
+								if( ComputeBarycentricWeightForPointOnPolygon( PolygonID, CenterVertexPosition, /* Out */ Triangle, /* Out */ TriangleVertexWeights ) )
 								{
 									for( int32 TextureCoordinateIndex = 0; TextureCoordinateIndex < TextureCoordinateCount; ++TextureCoordinateIndex )
 									{
 										VertexAndAttributes.PolygonVertexAttributes.Attributes.Add( FMeshElementAttributeData(
 											UEditableMeshAttribute::VertexTextureCoordinate(),
 											TextureCoordinateIndex,
-											TriangleVertexWeights.X * GetPolygonPerimeterVertexAttribute( PolygonRef, PerimeterVertexIndices[ 0 ], UEditableMeshAttribute::VertexTextureCoordinate(), TextureCoordinateIndex ) +
-											TriangleVertexWeights.Y * GetPolygonPerimeterVertexAttribute( PolygonRef, PerimeterVertexIndices[ 1 ], UEditableMeshAttribute::VertexTextureCoordinate(), TextureCoordinateIndex ) +
-											TriangleVertexWeights.Z * GetPolygonPerimeterVertexAttribute( PolygonRef, PerimeterVertexIndices[ 2 ], UEditableMeshAttribute::VertexTextureCoordinate(), TextureCoordinateIndex ) ) );
+											TriangleVertexWeights.X * GetVertexInstanceAttribute( Triangle.VertexInstanceIDs[ 0 ], UEditableMeshAttribute::VertexTextureCoordinate(), TextureCoordinateIndex ) +
+											TriangleVertexWeights.Y * GetVertexInstanceAttribute( Triangle.VertexInstanceIDs[ 1 ], UEditableMeshAttribute::VertexTextureCoordinate(), TextureCoordinateIndex ) +
+											TriangleVertexWeights.Z * GetVertexInstanceAttribute( Triangle.VertexInstanceIDs[ 2 ], UEditableMeshAttribute::VertexTextureCoordinate(), TextureCoordinateIndex ) ) );
 									}
 
 									VertexAndAttributes.PolygonVertexAttributes.Attributes.Add( FMeshElementAttributeData(
 										UEditableMeshAttribute::VertexColor(),
 										0,
-										TriangleVertexWeights.X * GetPolygonPerimeterVertexAttribute( PolygonRef, PerimeterVertexIndices[ 0 ], UEditableMeshAttribute::VertexColor(), 0 ) +
-										TriangleVertexWeights.Y * GetPolygonPerimeterVertexAttribute( PolygonRef, PerimeterVertexIndices[ 1 ], UEditableMeshAttribute::VertexColor(), 0 ) +
-										TriangleVertexWeights.Z * GetPolygonPerimeterVertexAttribute( PolygonRef, PerimeterVertexIndices[ 2 ], UEditableMeshAttribute::VertexColor(), 0 ) ) );
+										TriangleVertexWeights.X * GetVertexInstanceAttribute( Triangle.VertexInstanceIDs[ 0 ], UEditableMeshAttribute::VertexColor(), 0 ) +
+										TriangleVertexWeights.Y * GetVertexInstanceAttribute( Triangle.VertexInstanceIDs[ 1 ], UEditableMeshAttribute::VertexColor(), 0 ) +
+										TriangleVertexWeights.Z * GetVertexInstanceAttribute( Triangle.VertexInstanceIDs[ 2 ], UEditableMeshAttribute::VertexColor(), 0 ) ) );
 								}
 							}
 						}
@@ -6237,7 +8344,7 @@
 						{
 							// Transfer over UVs and vertex colors from the original polygon that we're replacing
 							{
-								const int32 VertexNumber = FindPolygonPerimeterVertexNumberForVertex( PolygonRef, VertexAndAttributes.VertexID );
+								const int32 VertexNumber = FindPolygonPerimeterVertexNumberForVertex( PolygonID, VertexAndAttributes.VertexID );
 								check( VertexNumber != INDEX_NONE );
 
 								for( int32 TextureCoordinateIndex = 0; TextureCoordinateIndex < TextureCoordinateCount; ++TextureCoordinateIndex )
@@ -6245,13 +8352,13 @@
 									VertexAndAttributes.PolygonVertexAttributes.Attributes.Add( FMeshElementAttributeData(
 										UEditableMeshAttribute::VertexTextureCoordinate(),
 										TextureCoordinateIndex,
-										GetPolygonPerimeterVertexAttribute( PolygonRef, VertexNumber, UEditableMeshAttribute::VertexTextureCoordinate(), TextureCoordinateIndex ) ) );
+										GetPolygonPerimeterVertexAttribute( PolygonID, VertexNumber, UEditableMeshAttribute::VertexTextureCoordinate(), TextureCoordinateIndex ) ) );
 								}
 
 								VertexAndAttributes.PolygonVertexAttributes.Attributes.Add( FMeshElementAttributeData(
 									UEditableMeshAttribute::VertexColor(),
 									0,
-									GetPolygonPerimeterVertexAttribute( PolygonRef, VertexNumber, UEditableMeshAttribute::VertexColor(), 0 ) ) );
+									GetPolygonPerimeterVertexAttribute( PolygonID, VertexNumber, UEditableMeshAttribute::VertexColor(), 0 ) ) );
 							}
 						}
 					}
@@ -6264,13 +8371,13 @@
 				// other three triangles will go in the corners of the original triangle.
 
 				static TArray<FVertexID> PerimeterVertexIDs;
-				GetPolygonPerimeterVertices( PolygonRef, PerimeterVertexIDs );
+				GetPolygonPerimeterVertices( PolygonID, PerimeterVertexIDs );
 				check( PerimeterVertexIDs.Num() == 6 );	// We split the triangle's 3 edges earlier, so we must have six edges now
 
 				// Define the new center triangle
 				{
 					FPolygonToCreate& PolygonToCreate = *new( PolygonsToCreate ) FPolygonToCreate();
-					PolygonToCreate.SectionID = PolygonRef.SectionID;
+					PolygonToCreate.PolygonGroupID = PolygonGroupID;
 
 					for( int32 OriginalVertexNumber = 0; OriginalVertexNumber < 3; ++OriginalVertexNumber )
 					{
@@ -6286,13 +8393,13 @@
 								VertexAndAttributes.PolygonVertexAttributes.Attributes.Add( FMeshElementAttributeData(
 									UEditableMeshAttribute::VertexTextureCoordinate(),
 									TextureCoordinateIndex,
-									GetPolygonPerimeterVertexAttribute( PolygonRef, VertexNumber, UEditableMeshAttribute::VertexTextureCoordinate(), TextureCoordinateIndex ) ) );
+									GetPolygonPerimeterVertexAttribute( PolygonID, VertexNumber, UEditableMeshAttribute::VertexTextureCoordinate(), TextureCoordinateIndex ) ) );
 							}
 
 							VertexAndAttributes.PolygonVertexAttributes.Attributes.Add( FMeshElementAttributeData(
 								UEditableMeshAttribute::VertexColor(),
 								0,
-								GetPolygonPerimeterVertexAttribute( PolygonRef, VertexNumber, UEditableMeshAttribute::VertexColor(), 0 ) ) );
+								GetPolygonPerimeterVertexAttribute( PolygonID, VertexNumber, UEditableMeshAttribute::VertexColor(), 0 ) ) );
 						}
 					}
 				}
@@ -6305,7 +8412,7 @@
 					const int32 NextVertexNumber = ( CurrentVertexNumber + 1 ) % PerimeterVertexIDs.Num();
 
 					FPolygonToCreate& PolygonToCreate = *new( PolygonsToCreate ) FPolygonToCreate();
-					PolygonToCreate.SectionID = PolygonRef.SectionID;
+					PolygonToCreate.PolygonGroupID = PolygonGroupID;
 
 					for( int32 TriangleVertexNumber = 0; TriangleVertexNumber < 3; ++TriangleVertexNumber )
 					{
@@ -6339,13 +8446,13 @@
 									VertexAndAttributes.PolygonVertexAttributes.Attributes.Add( FMeshElementAttributeData(
 										UEditableMeshAttribute::VertexTextureCoordinate(),
 										TextureCoordinateIndex,
-										GetPolygonPerimeterVertexAttribute( PolygonRef, VertexNumber, UEditableMeshAttribute::VertexTextureCoordinate(), TextureCoordinateIndex ) ) );
+										GetPolygonPerimeterVertexAttribute( PolygonID, VertexNumber, UEditableMeshAttribute::VertexTextureCoordinate(), TextureCoordinateIndex ) ) );
 								}
 
 								VertexAndAttributes.PolygonVertexAttributes.Attributes.Add( FMeshElementAttributeData(
 									UEditableMeshAttribute::VertexColor(),
 									0,
-									GetPolygonPerimeterVertexAttribute( PolygonRef, VertexNumber, UEditableMeshAttribute::VertexColor(), 0 ) ) );
+									GetPolygonPerimeterVertexAttribute( PolygonID, VertexNumber, UEditableMeshAttribute::VertexColor(), 0 ) ) );
 							}
 						}
 					}
@@ -6359,30 +8466,44 @@
 	{
 		const bool bDeleteOrphanedEdges = false;	// No need to delete orphans, because this function won't orphan anything
 		const bool bDeleteOrphanedVertices = false;
-		const bool bDeleteEmptySections = false;
-		DeletePolygons( PolygonRefs, bDeleteOrphanedEdges, bDeleteOrphanedVertices, bDeleteEmptySections );
+		const bool bDeleteOrphanedVertexInstances = false;
+		const bool bDeleteEmptyPolygonGroups = false;
+		DeletePolygons( PolygonIDs, bDeleteOrphanedEdges, bDeleteOrphanedVertices, bDeleteOrphanedVertexInstances, bDeleteEmptyPolygonGroups );
 	}
 
 
 	// Create all of the new polygons for the tessellated representation of the original polygons
 	{
 		static TArray<FEdgeID> NewEdgeIDs;
-		CreatePolygons( PolygonsToCreate, /* Out */ OutNewPolygonRefs, /* Out */ NewEdgeIDs );
+		CreatePolygons( PolygonsToCreate, /* Out */ OutNewPolygonIDs, /* Out */ NewEdgeIDs );
 	}
-
-
-	// Generate normals for all of the new polygons
-	GenerateNormalsAndTangentsForPolygons( OutNewPolygonRefs );
 }
 
 
-void UEditableMesh::SetTextureCoordinateCount(int32 NumTexCoords)
+void UEditableMesh::SetTextureCoordinateCount( const int32 NumTexCoords )
 {
-	TextureCoordinateCount = FMath::Max(NumTexCoords, 0);
+	TextureCoordinateCount = FMath::Max( NumTexCoords, 0 );
 }
 
 
-void UEditableMesh::QuadrangulateMesh( TArray<FPolygonRef>& NewPolygonRefs )
+void UEditableMesh::QuadrangulateMesh( TArray<FPolygonID>& OutNewPolygonIDs )
 {
+	// Iterate through all polygon groups in the mesh, quadrangulating each in turn
+	const int32 MaxPolygonGroupIndex = GetPolygonGroupArraySize();
+	for( int32 PolygonGroupIndex = 0; PolygonGroupIndex < MaxPolygonGroupIndex; ++PolygonGroupIndex )
+	{
+		const FPolygonGroupID PolygonGroupID( PolygonGroupIndex );
+		if( IsValidPolygonGroup( PolygonGroupID ) )
+		{
+			static TArray<FPolygonID> NewPolygonIDs;
+			QuadrangulatePolygonGroup( PolygonGroupID, NewPolygonIDs );
+			OutNewPolygonIDs.Append( NewPolygonIDs );
+		}
+	}
+}
+
+
+void UEditableMesh::QuadrangulatePolygonGroup( const FPolygonGroupID PolygonGroupID, TArray<FPolygonID>& OutNewPolygonIDs )
+{
 	// Uses the first two steps of the algorithm described by
 	// http://www.lirmm.fr/~beniere/ArticlesPersos/GRAPP10_Beniere_Final.pdf
@@ -6393,56 +8514,43 @@
 	const bool bKeepTextureBorder = true;
 	const bool bKeepColorBorder = true;
 
-	NewPolygonRefs.Reset();
+	OutNewPolygonIDs.Reset();
 
-	static TArray<FPolygonRef> PolygonRefs;
-	PolygonRefs.Reset();
+	static TArray<FPolygonID> PolygonIDs;
+	PolygonIDs.Reset();
 
-	static TSet<FPolygonRef> DeletedPolygonRefs;
-	DeletedPolygonRefs.Reset();
-
 	// Get a list of all polygon refs in the mesh
+	// @todo mesheditor: provide direct access to the polygons array in the group
 	{
-		const int32 MaxSectionIndex = GetSectionArraySize();
-		for( int32 SectionIndex = 0; SectionIndex < MaxSectionIndex; ++SectionIndex )
+		const int32 PolygonCountInGroup = GetPolygonCountInGroup( PolygonGroupID );
+		PolygonIDs.Reserve( PolygonCountInGroup );
+		for( int32 PolygonIndex = 0; PolygonIndex < PolygonCountInGroup; ++PolygonIndex )
 		{
-			const FSectionID SectionID( SectionIndex );
-			if( IsValidSection( SectionID ) )
-			{
-				const int32 MaxPolygonIndex = GetPolygonArraySize( SectionID );
-				for( int32 PolygonIndex = 0; PolygonIndex < MaxPolygonIndex; ++PolygonIndex )
-				{
-					const FPolygonRef PolygonRef( SectionID, FPolygonID( PolygonIndex ) );
-					if( IsValidPolygon( PolygonRef ) )
-					{
-						PolygonRefs.Add( PolygonRef );
-					}
-				}
-			}
+			PolygonIDs.Add( GetPolygonInGroup( PolygonGroupID, PolygonIndex ) );
 		}
 	}
 
 	// This represents an adjacent triangle which can be merged to a quadrilateral, and an assigned score based on the 'quality' of the resulting quadrilateral.
 	struct FAdjacentPolygon
 	{
-		// PolygonRef of the adjacent triangle. This object is keyed on 'our' PolygonRef.
-		FPolygonRef PolygonRef;
+		// Polygon ID of the adjacent triangle. This object is keyed on 'our' Polygon ID.
+		FPolygonID PolygonID;
 
-		// List of the vertices which form the merged quadrilateral - (PolygonRef, PolygonPerimeterIndex)
-		TTuple<FPolygonRef, int32> Vertices[ 4 ];
+		// List of the vertices which form the merged quadrilateral - (Polygon ID, PolygonPerimeterIndex)
+		TTuple<FPolygonID, int32> Vertices[ 4 ];
 
 		// 'Quality' of the quadrilateral (internal angles closer to 90 degrees are better)
 		float Score;
 
 		FAdjacentPolygon() {}
 
-		FAdjacentPolygon( FPolygonRef InPolygonRef, FPolygonRef InAdjacentPolygonRef, int32 InVertex0, int32 InVertex1, int32 InVertex2, int32 InVertex3, float InScore )
+		FAdjacentPolygon( FPolygonID InPolygonID, FPolygonID InAdjacentPolygonID, int32 InVertex0, int32 InVertex1, int32 InVertex2, int32 InVertex3, float InScore )
 		{
-			PolygonRef = InAdjacentPolygonRef;
-			Vertices[ 0 ] = MakeTuple( InPolygonRef, InVertex0 );
-			Vertices[ 1 ] = MakeTuple( InPolygonRef, InVertex1 );
-			Vertices[ 2 ] = MakeTuple( InAdjacentPolygonRef, InVertex2 );
-			Vertices[ 3 ] = MakeTuple( InAdjacentPolygonRef, InVertex3 );
+			PolygonID = InAdjacentPolygonID;
+			Vertices[ 0 ] = MakeTuple( InPolygonID, InVertex0 );
+			Vertices[ 1 ] = MakeTuple( InPolygonID, InVertex1 );
+			Vertices[ 2 ] = MakeTuple( InAdjacentPolygonID, InVertex2 );
+			Vertices[ 3 ] = MakeTuple( InAdjacentPolygonID, InVertex3 );
 			Score = InScore;
 		}
 	};
@@ -6492,11 +8600,11 @@
 			return AdjacentPolygons[ 0 ];
 		}
 		
-		bool Remove( const FPolygonRef PolygonRef )
+		bool Remove( const FPolygonID PolygonID )
 		{
 			for( int32 Index = 0; Index < NumAdjacentPolygons; ++Index )
 			{
-				if( AdjacentPolygons[ Index ].PolygonRef == PolygonRef )
+				if( AdjacentPolygons[ Index ].PolygonID == PolygonID )
 				{
 					for( int32 CopyIndex = Index + 1; CopyIndex < NumAdjacentPolygons; ++CopyIndex )
 					{
@@ -6511,11 +8619,11 @@
 			return false;
 		}
 
-		bool Contains( const FPolygonRef PolygonRef ) const
+		bool Contains( const FPolygonID PolygonID ) const
 		{
 			for( int32 Index = 0; Index < NumAdjacentPolygons; ++Index )
 			{
-				if( AdjacentPolygons[ Index ].PolygonRef == PolygonRef )
+				if( AdjacentPolygons[ Index ].PolygonID == PolygonID )
 				{
 					return true;
 				}
@@ -6524,10 +8632,10 @@
 			return false;
 		}
 
-		FPolygonRef GetPolygonRef( int32 Index ) const
+		FPolygonID GetPolygonID( int32 Index ) const
 		{
 			check( Index < NumAdjacentPolygons );
-			return AdjacentPolygons[ Index ].PolygonRef;
+			return AdjacentPolygons[ Index ].PolygonID;
 		}
 
 		int32 Num() const { return NumAdjacentPolygons; }
@@ -6537,34 +8645,34 @@
 
 	// Build list of valid adjacent triangle pairs, and assign a score based on the quality of the quadrilateral they form
 
-	static TMap<FPolygonRef, FAdjacentPolygons> AdjacentPolygonsMap;
+	static TMap<FPolygonID, FAdjacentPolygons> AdjacentPolygonsMap;
 	AdjacentPolygonsMap.Reset();
 
-	FPolygonRef PolygonRefToMerge1 = FPolygonRef::Invalid;
+	FPolygonID PolygonIDToMerge1 = FPolygonID::Invalid;
 	{
 		float BestScore = TNumericLimits<float>::Max();
 
-		for( const FPolygonRef PolygonRef : PolygonRefs )
+		for( const FPolygonID PolygonID : PolygonIDs )
 		{
 			// If it's not a triangle, don't consider this polygon at all
-			if( GetPolygonPerimeterEdgeCount( PolygonRef ) != 3 )
+			if( GetPolygonPerimeterEdgeCount( PolygonID ) != 3 )
 			{
 				continue;
 			}
 
 			// We're only interested in adjacent triangles which are nearly coplanar; get the normal so we can compare it with the adjacent polygons' normals
-			const FVector PolygonNormal = ComputePolygonNormal( PolygonRef );
+			const FVector PolygonNormal = ComputePolygonNormal( PolygonID );
 
 			// Go round the edge considering all adjacent polygons, looking for valid pairs and assigning a quality score (lower is better)
 			for( int32 PerimeterEdgeIndex = 0; PerimeterEdgeIndex < 3; ++PerimeterEdgeIndex )
 			{
 				bool bOutEdgeWindingIsReversedForPolygon;
-				const FEdgeID PerimeterEdgeID = GetPolygonPerimeterEdge( PolygonRef, PerimeterEdgeIndex, bOutEdgeWindingIsReversedForPolygon );
+				const FEdgeID PerimeterEdgeID = GetPolygonPerimeterEdge( PolygonID, PerimeterEdgeIndex, bOutEdgeWindingIsReversedForPolygon );
 
 				const bool bIsSoftEdge = FMath::IsNearlyZero( GetEdgeAttribute( PerimeterEdgeID, UEditableMeshAttribute::EdgeIsHard(), 0 ).X );
 				if( !bKeepHardEdges || bIsSoftEdge )
 				{
-					const FPolygonRef AdjacentPolygonRef = [ this, PerimeterEdgeID, PolygonRef ]() -> FPolygonRef
+					const FPolygonID AdjacentPolygonID = [ this, PerimeterEdgeID, PolygonID ]() -> FPolygonID
 					{
 						const int32 EdgeConnectedPolygonCount = GetEdgeConnectedPolygonCount( PerimeterEdgeID );
 
@@ -6573,23 +8681,24 @@
 						{
 							for( int32 EdgeConnectedPolygonIndex = 0; EdgeConnectedPolygonIndex < 2; ++EdgeConnectedPolygonIndex )
 							{
-								const FPolygonRef EdgeConnectedPolygonRef = GetEdgeConnectedPolygon( PerimeterEdgeID, EdgeConnectedPolygonIndex );
-								if( EdgeConnectedPolygonRef != PolygonRef )
+								const FPolygonID EdgeConnectedPolygonID = GetEdgeConnectedPolygon( PerimeterEdgeID, EdgeConnectedPolygonIndex );
+								if( EdgeConnectedPolygonID != PolygonID )
 								{
-									return GetPolygonPerimeterEdgeCount( EdgeConnectedPolygonRef ) == 3 ? EdgeConnectedPolygonRef : FPolygonRef::Invalid;
+									return GetPolygonPerimeterEdgeCount( EdgeConnectedPolygonID ) == 3 ? EdgeConnectedPolygonID : FPolygonID::Invalid;
 								}
 							}
 						}
 
-						return FPolygonRef::Invalid;
+						return FPolygonID::Invalid;
 					}();
 
-					if( AdjacentPolygonRef != FPolygonRef::Invalid )
+					if( AdjacentPolygonID != FPolygonID::Invalid )
 					{
-						FAdjacentPolygons* AdjacentPolygons = AdjacentPolygonsMap.Find( PolygonRef );
-						if( !AdjacentPolygons || !AdjacentPolygons->Contains( AdjacentPolygonRef ) )
+						FAdjacentPolygons* AdjacentPolygons = AdjacentPolygonsMap.Find( PolygonID );
+
+						if( !AdjacentPolygons || !AdjacentPolygons->Contains( AdjacentPolygonID ) )
 						{
-							const FVector AdjacentPolygonNormal = ComputePolygonNormal( AdjacentPolygonRef );
+							const FVector AdjacentPolygonNormal = ComputePolygonNormal( AdjacentPolygonID );
 							const float AdjacentPolygonDot = FVector::DotProduct( PolygonNormal, AdjacentPolygonNormal );
 
 							if( AdjacentPolygonDot >= CosAngleThreshold )
@@ -6616,8 +8725,8 @@
 								const int32 PrevPerimeterEdgeIndex = ( PerimeterEdgeIndex + 2 ) % 3;
 								const int32 NextPerimeterEdgeIndex = ( PerimeterEdgeIndex + 1 ) % 3;
 
-								const FVertexID SharedVertexID = GetPolygonPerimeterVertex( PolygonRef, PerimeterEdgeIndex );
-								const int32 AdjacentPerimeterEdgeIndex = FindPolygonPerimeterVertexNumberForVertex( AdjacentPolygonRef, SharedVertexID );
+								const FVertexID SharedVertexID = GetPolygonPerimeterVertex( PolygonID, PerimeterEdgeIndex );
+								const int32 AdjacentPerimeterEdgeIndex = FindPolygonPerimeterVertexNumberForVertex( AdjacentPolygonID, SharedVertexID );
 
 								const int32 PrevAdjacentPerimeterEdgeIndex = ( AdjacentPerimeterEdgeIndex + 2 ) % 3;
 								const int32 NextAdjacentPerimeterEdgeIndex = ( AdjacentPerimeterEdgeIndex + 1 ) % 3;
@@ -6628,10 +8737,10 @@
 								{
 									for( int32 TextureCoordinateIndex = 0; TextureCoordinateIndex < TextureCoordinateCount; ++TextureCoordinateIndex )
 									{
-										const FVector4 UVPolygon1Vertex1 = GetPolygonPerimeterVertexAttribute( PolygonRef, PerimeterEdgeIndex, UEditableMeshAttribute::VertexTextureCoordinate(), TextureCoordinateIndex );
-										const FVector4 UVPolygon2Vertex1 = GetPolygonPerimeterVertexAttribute( AdjacentPolygonRef, AdjacentPerimeterEdgeIndex, UEditableMeshAttribute::VertexTextureCoordinate(), TextureCoordinateIndex );
-										const FVector4 UVPolygon1Vertex2 = GetPolygonPerimeterVertexAttribute( PolygonRef, NextPerimeterEdgeIndex, UEditableMeshAttribute::VertexTextureCoordinate(), TextureCoordinateIndex );
-										const FVector4 UVPolygon2Vertex2 = GetPolygonPerimeterVertexAttribute( AdjacentPolygonRef, PrevAdjacentPerimeterEdgeIndex, UEditableMeshAttribute::VertexTextureCoordinate(), TextureCoordinateIndex );
+										const FVector4 UVPolygon1Vertex1 = GetPolygonPerimeterVertexAttribute( PolygonID, PerimeterEdgeIndex, UEditableMeshAttribute::VertexTextureCoordinate(), TextureCoordinateIndex );
+										const FVector4 UVPolygon2Vertex1 = GetPolygonPerimeterVertexAttribute( AdjacentPolygonID, AdjacentPerimeterEdgeIndex, UEditableMeshAttribute::VertexTextureCoordinate(), TextureCoordinateIndex );
+										const FVector4 UVPolygon1Vertex2 = GetPolygonPerimeterVertexAttribute( PolygonID, NextPerimeterEdgeIndex, UEditableMeshAttribute::VertexTextureCoordinate(), TextureCoordinateIndex );
+										const FVector4 UVPolygon2Vertex2 = GetPolygonPerimeterVertexAttribute( AdjacentPolygonID, PrevAdjacentPerimeterEdgeIndex, UEditableMeshAttribute::VertexTextureCoordinate(), TextureCoordinateIndex );
 										if( UVPolygon1Vertex1 != UVPolygon2Vertex1 || UVPolygon1Vertex2 != UVPolygon2Vertex2 )
 										{
 											bTextureCoordinatesEqual = false;
@@ -6644,10 +8753,10 @@
 								bool bColorsEqual = true;
 								if( bKeepColorBorder )
 								{
-									const FVector4 ColorPolygon1Vertex1 = GetPolygonPerimeterVertexAttribute( PolygonRef, PerimeterEdgeIndex, UEditableMeshAttribute::VertexColor(), 0 );
-									const FVector4 ColorPolygon2Vertex1 = GetPolygonPerimeterVertexAttribute( AdjacentPolygonRef, AdjacentPerimeterEdgeIndex, UEditableMeshAttribute::VertexColor(), 0 );
-									const FVector4 ColorPolygon1Vertex2 = GetPolygonPerimeterVertexAttribute( PolygonRef, NextPerimeterEdgeIndex, UEditableMeshAttribute::VertexColor(), 0 );
-									const FVector4 ColorPolygon2Vertex2 = GetPolygonPerimeterVertexAttribute( AdjacentPolygonRef, PrevAdjacentPerimeterEdgeIndex, UEditableMeshAttribute::VertexColor(), 0 );
+									const FVector4 ColorPolygon1Vertex1 = GetPolygonPerimeterVertexAttribute( PolygonID, PerimeterEdgeIndex, UEditableMeshAttribute::VertexColor(), 0 );
+									const FVector4 ColorPolygon2Vertex1 = GetPolygonPerimeterVertexAttribute( AdjacentPolygonID, AdjacentPerimeterEdgeIndex, UEditableMeshAttribute::VertexColor(), 0 );
+									const FVector4 ColorPolygon1Vertex2 = GetPolygonPerimeterVertexAttribute( PolygonID, NextPerimeterEdgeIndex, UEditableMeshAttribute::VertexColor(), 0 );
+									const FVector4 ColorPolygon2Vertex2 = GetPolygonPerimeterVertexAttribute( AdjacentPolygonID, PrevAdjacentPerimeterEdgeIndex, UEditableMeshAttribute::VertexColor(), 0 );
 									if( ColorPolygon1Vertex1 != ColorPolygon2Vertex1 || ColorPolygon1Vertex2 != ColorPolygon2Vertex2 )
 									{
 										bColorsEqual = false;
@@ -6657,11 +8766,11 @@
 								// If we are maintaining texture or color seams, only merge the polygons if the appropriate attributes are equal for each polygon
 								if( ( !bKeepTextureBorder || bTextureCoordinatesEqual ) && ( !bKeepColorBorder || bColorsEqual ) )
 								{
-									const FVertexID V0 = GetPolygonPerimeterVertex( PolygonRef, PrevPerimeterEdgeIndex );
-									const FVertexID V1 = GetPolygonPerimeterVertex( PolygonRef, PerimeterEdgeIndex );
-									const FVertexID V2 = GetPolygonPerimeterVertex( AdjacentPolygonRef, NextAdjacentPerimeterEdgeIndex );
-									const FVertexID V3 = GetPolygonPerimeterVertex( AdjacentPolygonRef, PrevAdjacentPerimeterEdgeIndex );
-									check( V3 == GetPolygonPerimeterVertex( PolygonRef, NextPerimeterEdgeIndex ) );
+									const FVertexID V0 = GetPolygonPerimeterVertex( PolygonID, PrevPerimeterEdgeIndex );
+									const FVertexID V1 = GetPolygonPerimeterVertex( PolygonID, PerimeterEdgeIndex );
+									const FVertexID V2 = GetPolygonPerimeterVertex( AdjacentPolygonID, NextAdjacentPerimeterEdgeIndex );
+									const FVertexID V3 = GetPolygonPerimeterVertex( AdjacentPolygonID, PrevAdjacentPerimeterEdgeIndex );
+									check( V3 == GetPolygonPerimeterVertex( PolygonID, NextPerimeterEdgeIndex ) );
 
 									const FVector P0 = GetVertexAttribute( V0, UEditableMeshAttribute::VertexPosition(), 0 );
 									const FVector P1 = GetVertexAttribute( V1, UEditableMeshAttribute::VertexPosition(), 0 );
@@ -6685,14 +8794,14 @@
 									if( Score < BestScore )
 									{
 										BestScore = Score;
-										PolygonRefToMerge1 = PolygonRef;
+										PolygonIDToMerge1 = PolygonID;
 									}
 
 									// Add to a list of adjacent polygons, sorted by score
-									FAdjacentPolygons& AdjacentPolygons1 = AdjacentPolygonsMap.FindOrAdd( PolygonRef );
+									FAdjacentPolygons& AdjacentPolygons1 = AdjacentPolygonsMap.FindOrAdd( PolygonID );
 									AdjacentPolygons1.Add( FAdjacentPolygon(
-										PolygonRef,
-										AdjacentPolygonRef,
+										PolygonID,
+										AdjacentPolygonID,
 										PrevPerimeterEdgeIndex,
 										PerimeterEdgeIndex,
 										NextAdjacentPerimeterEdgeIndex,
@@ -6700,11 +8809,11 @@
 										Score ) );
 
 									// And perform the corresponding operation the other way round
-									FAdjacentPolygons& AdjacentPolygons2 = AdjacentPolygonsMap.FindOrAdd( AdjacentPolygonRef );
-									check( !AdjacentPolygons2.Contains( PolygonRef ) );
+									FAdjacentPolygons& AdjacentPolygons2 = AdjacentPolygonsMap.FindOrAdd( AdjacentPolygonID );
+									check( !AdjacentPolygons2.Contains( PolygonID ) );
 									AdjacentPolygons2.Add( FAdjacentPolygon(
-										AdjacentPolygonRef,
-										PolygonRef,
+										AdjacentPolygonID,
+										PolygonID,
 										NextAdjacentPerimeterEdgeIndex,
 										PrevAdjacentPerimeterEdgeIndex,
 										PrevPerimeterEdgeIndex,
@@ -6719,7 +8828,8 @@
 		}
 	}
 
-	if( PolygonRefToMerge1 == FPolygonRef::Invalid )
+	// If there were no valid pairs of polys to merge, finish already
+	if( PolygonIDToMerge1 == FPolygonID::Invalid )
 	{
 		return;
 	}
@@ -6727,20 +8837,21 @@
 	static TArray<FPolygonToCreate> PolygonsToCreate;
 	PolygonsToCreate.Reset();
 
-	static TArray<FPolygonRef> PolygonRefsToDelete;
-	PolygonRefsToDelete.Reset();
+	static TArray<FPolygonID> PolygonIDsToDelete;
+	PolygonIDsToDelete.Reset();
 
-	TSet<FPolygonRef> BoundaryPolygons;
+	static TSet<FPolygonID> BoundaryPolygons;
+	BoundaryPolygons.Reset();
 
 	// Propagate quadrangulated area outwards from starting polygon
 	for( ; ; )
 	{
-		FAdjacentPolygons& AdjacentPolygons1 = AdjacentPolygonsMap.FindChecked( PolygonRefToMerge1 );
+		FAdjacentPolygons& AdjacentPolygons1 = AdjacentPolygonsMap.FindChecked( PolygonIDToMerge1 );
 		check( AdjacentPolygons1.IsValid() );
 		const FAdjacentPolygon& AdjacentPolygon1 = AdjacentPolygons1.GetBestAdjacentPolygon();
-		const FPolygonRef PolygonRefToMerge2 = AdjacentPolygon1.PolygonRef;
+		const FPolygonID PolygonIDToMerge2 = AdjacentPolygon1.PolygonID;
 
-		FAdjacentPolygons& AdjacentPolygons2 = AdjacentPolygonsMap.FindChecked( PolygonRefToMerge2 );
+		FAdjacentPolygons& AdjacentPolygons2 = AdjacentPolygonsMap.FindChecked( PolygonIDToMerge2 );
 		check( AdjacentPolygons2.IsValid() );
 
 		// Create new quadrilateral
@@ -6749,7 +8860,7 @@
 
 		FPolygonToCreate& PolygonToCreate = PolygonsToCreate.Last();
 
-		PolygonToCreate.SectionID = PolygonRefToMerge1.SectionID;
+		PolygonToCreate.PolygonGroupID = PolygonGroupID;
 		PolygonToCreate.PerimeterVertices.Reset( 4 );
 
 		for( int32 VertexIndex = 0; VertexIndex < 4; ++VertexIndex )
@@ -6789,74 +8900,74 @@
 
 		// Specify old polygons to be deleted
 
-		check( !PolygonRefsToDelete.Contains( PolygonRefToMerge1 ) );
-		check( !PolygonRefsToDelete.Contains( PolygonRefToMerge2 ) );
-		PolygonRefsToDelete.Emplace( PolygonRefToMerge1 );
-		PolygonRefsToDelete.Emplace( PolygonRefToMerge2 );
+		check( !PolygonIDsToDelete.Contains( PolygonIDToMerge1 ) );
+		check( !PolygonIDsToDelete.Contains( PolygonIDToMerge2 ) );
+		PolygonIDsToDelete.Emplace( PolygonIDToMerge1 );
+		PolygonIDsToDelete.Emplace( PolygonIDToMerge2 );
 
 		// And remove them from the boundary set
 
-		BoundaryPolygons.Remove( PolygonRefToMerge1 );
-		BoundaryPolygons.Remove( PolygonRefToMerge2 );
+		BoundaryPolygons.Remove( PolygonIDToMerge1 );
+		BoundaryPolygons.Remove( PolygonIDToMerge2 );
 
 		// Now break connections between newly added polygons and their neighbors.
 		// If a polygon ends up with no connections, delete it entirely from the map so it is no longer considered.
 		// This happens if a polygon has been added to the quadrangulated set, or if it is an orphaned triangle which cannot be paired to anything.
 		// We defer deleting the entry from the map until we have broken all connections.
 
-		verify( AdjacentPolygons1.Remove( PolygonRefToMerge2 ) );
-		verify( AdjacentPolygons2.Remove( PolygonRefToMerge1 ) );
+		verify( AdjacentPolygons1.Remove( PolygonIDToMerge2 ) );
+		verify( AdjacentPolygons2.Remove( PolygonIDToMerge1 ) );
 
-		static TArray<FPolygonRef> AdjacentPolygonsEntryToDelete;
+		static TArray<FPolygonID> AdjacentPolygonsEntryToDelete;
 		AdjacentPolygonsEntryToDelete.Reset();
 
 		for( int32 Index = 0; Index < AdjacentPolygons1.Num(); ++Index )
 		{
-			const FPolygonRef AdjacentPolygonRef = AdjacentPolygons1.GetPolygonRef( Index );
+			const FPolygonID AdjacentPolygonID = AdjacentPolygons1.GetPolygonID( Index );
 
-			FAdjacentPolygons* OtherAdjacentPolygons = AdjacentPolygonsMap.Find( AdjacentPolygonRef );
+			FAdjacentPolygons* OtherAdjacentPolygons = AdjacentPolygonsMap.Find( AdjacentPolygonID );
 			if( OtherAdjacentPolygons )
 			{
-				verify( OtherAdjacentPolygons->Remove( PolygonRefToMerge1 ) );
+				verify( OtherAdjacentPolygons->Remove( PolygonIDToMerge1 ) );
 				if( !OtherAdjacentPolygons->IsValid() )
 				{
-					AdjacentPolygonsEntryToDelete.Add( AdjacentPolygonRef );
+					AdjacentPolygonsEntryToDelete.Add( AdjacentPolygonID );
 				}
 				else
 				{
-					BoundaryPolygons.Add( AdjacentPolygonRef );
+					BoundaryPolygons.Add( AdjacentPolygonID );
 				}
 			}
 		}
 
-		AdjacentPolygonsEntryToDelete.Add( PolygonRefToMerge1 );
+		AdjacentPolygonsEntryToDelete.Add( PolygonIDToMerge1 );
 
 		for( int32 Index = 0; Index < AdjacentPolygons2.Num(); ++Index )
 		{
-			const FPolygonRef AdjacentPolygonRef = AdjacentPolygons2.GetPolygonRef( Index );
+			const FPolygonID AdjacentPolygonID = AdjacentPolygons2.GetPolygonID( Index );
 
-			FAdjacentPolygons* OtherAdjacentPolygons = AdjacentPolygonsMap.Find( AdjacentPolygonRef );
+			FAdjacentPolygons* OtherAdjacentPolygons = AdjacentPolygonsMap.Find( AdjacentPolygonID );
 			if( OtherAdjacentPolygons )
 			{
-				verify( OtherAdjacentPolygons->Remove( PolygonRefToMerge2 ) );
+				verify( OtherAdjacentPolygons->Remove( PolygonIDToMerge2 ) );
 				if( !OtherAdjacentPolygons->IsValid() )
 				{
-					AdjacentPolygonsEntryToDelete.Add( AdjacentPolygonRef );
+					AdjacentPolygonsEntryToDelete.Add( AdjacentPolygonID );
 				}
 				else
 				{
-					BoundaryPolygons.Add( AdjacentPolygonRef );
+					BoundaryPolygons.Add( AdjacentPolygonID );
 				}
 			}
 		}
 
-		AdjacentPolygonsEntryToDelete.Add( PolygonRefToMerge2 );
+		AdjacentPolygonsEntryToDelete.Add( PolygonIDToMerge2 );
 
 		// Clean up: any polygons' map entries which now have no adjacent polygons get deleted completely.
 		// This implies they have no connected neighbors which can be merged (either because they are near an edge with only unmergeable polygons nearby, or because
 		// they are in the middle of the quadrangulated area)
 
-		for( const FPolygonRef AdjacentPolygonsEntry : AdjacentPolygonsEntryToDelete )
+		for( const FPolygonID AdjacentPolygonsEntry : AdjacentPolygonsEntryToDelete )
 		{
 			AdjacentPolygonsMap.Remove( AdjacentPolygonsEntry );
 			BoundaryPolygons.Remove( AdjacentPolygonsEntry );
@@ -6865,27 +8976,27 @@
 		// Now look for the next polygon to use: it is the one with the best score from the BoundaryPolygons set.
 
 		float BestScore = TNumericLimits<float>::Max();
-		PolygonRefToMerge1 = FPolygonRef::Invalid;
-		for( const FPolygonRef BoundaryPolygon : BoundaryPolygons )
+		PolygonIDToMerge1 = FPolygonID::Invalid;
+		for( const FPolygonID BoundaryPolygon : BoundaryPolygons )
 		{
 			const FAdjacentPolygons& AdjacentPolygons = AdjacentPolygonsMap.FindChecked( BoundaryPolygon );
 			const FAdjacentPolygon& AdjacentPolygon = AdjacentPolygons.GetBestAdjacentPolygon();
 			if( AdjacentPolygon.Score < BestScore )
 			{
-				PolygonRefToMerge1 = AdjacentPolygon.PolygonRef;
+				PolygonIDToMerge1 = AdjacentPolygon.PolygonID;
 			}
 		}
 
 		// If there are still no candidates adjacent to the already quadrangulated area, choose the best candidate elsewhere.
 		// This will start a new quadrangulated area, which is grown in the same way as the last.
 
-		if( PolygonRefToMerge1 == FPolygonRef::Invalid )
+		if( PolygonIDToMerge1 == FPolygonID::Invalid )
 		{
 			BoundaryPolygons.Reset();
 
 			for( const auto& AdjacentPolygonsMapEntry : AdjacentPolygonsMap )
 			{
-				const FPolygonRef PolygonRef = AdjacentPolygonsMapEntry.Key;
+				const FPolygonID PolygonID = AdjacentPolygonsMapEntry.Key;
 				const FAdjacentPolygons& AdjacentPolygons = AdjacentPolygonsMapEntry.Value;
 				check( AdjacentPolygons.IsValid() );
 
@@ -6894,14 +9005,14 @@
 				if( AdjacentPolygon.Score < BestScore )
 				{
 					BestScore = AdjacentPolygon.Score;
-					PolygonRefToMerge1 = PolygonRef;
+					PolygonIDToMerge1 = PolygonID;
 				}
 			}
 		}
 
 		// If there are still no candidates, we've done as much as we can do
 
-		if( PolygonRefToMerge1 == FPolygonRef::Invalid )
+		if( PolygonIDToMerge1 == FPolygonID::Invalid )
 		{
 			break;
 		}
@@ -6909,18 +9020,15 @@
 
 	// Finally, actually change the geometry and rebuild normals/tangents
 
-	static TArray<FPolygonRef> CreatedPolygonRefs;
-	static TArray< FEdgeID> CreatedEdgeIDs;
+	static TArray<FEdgeID> CreatedEdgeIDs;
 
-	CreatePolygons( PolygonsToCreate, CreatedPolygonRefs, CreatedEdgeIDs );
-	NewPolygonRefs.Append( CreatedPolygonRefs );
+	CreatePolygons( PolygonsToCreate, OutNewPolygonIDs, CreatedEdgeIDs );
 
 	const bool bDeleteOrphanedEdges = true;
 	const bool bDeleteOrphanedVertices = false;
-	const bool bDeleteEmptySections = false;
-	DeletePolygons( PolygonRefsToDelete, bDeleteOrphanedEdges, bDeleteOrphanedVertices, bDeleteEmptySections );
-
-	GenerateNormalsAndTangentsForPolygonsAndAdjacents( NewPolygonRefs );
+	const bool bDeleteOrphanedVertexInstances = false;
+	const bool bDeleteEmptyPolygonGroups = false;
+	DeletePolygons( PolygonIDsToDelete, bDeleteOrphanedEdges, bDeleteOrphanedVertices, bDeleteOrphanedVertexInstances, bDeleteEmptyPolygonGroups );
 }
 
 
@@ -6934,6 +9042,7 @@
 {
 	if( NewUndo.IsValid() )
 	{
+//		UE_LOG(LogEditableMesh, Verbose, TEXT("%s"), *NewUndo->ToString());
 		if( bAllowUndo )
 		{
 			if( !Undo.IsValid() )
