Copying //Tasks/UE4/Dev-Geometry-Refactor to Dev-Geometry (//UE4/Dev-Geometry)

diff --git a/Engine/Source/Runtime/MeshEditingRuntime/EditableMeshChanges.cpp b/Engine/Source/Runtime/MeshEditingRuntime/EditableMeshChanges.cpp
--- a/Engine/Source/Runtime/MeshEditingRuntime/EditableMeshChanges.cpp
+++ b/Engine/Source/Runtime/MeshEditingRuntime/EditableMeshChanges.cpp
@@ -23,6 +23,25 @@ FString FDeleteOrphanVerticesChange::ToString() const
 }
 
 
+TUniquePtr<FChange> FDeleteVertexInstancesChange::Execute( UObject* Object )
+{
+	UEditableMesh* EditableMesh = CastChecked<UEditableMesh>( Object );
+	verify( !EditableMesh->AnyChangesToUndo() );
+
+	EditableMesh->DeleteVertexInstances( Input.VertexInstanceIDsToDelete, Input.bDeleteOrphanedVertices );
+
+	return EditableMesh->MakeUndo();
+}
+
+
+FString FDeleteVertexInstancesChange::ToString() const
+{
+	return FString::Printf(
+		TEXT( "Delete Vertex Instances [VertexInstanceIDsToDelete:%s]" ),
+		*LogHelpers::ArrayToString( Input.VertexInstanceIDsToDelete ) );
+}
+
+
 TUniquePtr<FChange> FDeleteEdgesChange::Execute( UObject* Object )
 {
 	UEditableMesh* EditableMesh = CastChecked<UEditableMesh>( Object );
@@ -63,6 +82,26 @@ FString FCreateVerticesChange::ToString() const
 }
 
 
+TUniquePtr<FChange> FCreateVertexInstancesChange::Execute( UObject* Object )
+{
+	UEditableMesh* EditableMesh = CastChecked<UEditableMesh>( Object );
+	verify( !EditableMesh->AnyChangesToUndo() );
+
+	static TArray<FVertexInstanceID> UnusedNewVertexInstanceIDs;
+	EditableMesh->CreateVertexInstances( Input.VertexInstancesToCreate, /* Out */ UnusedNewVertexInstanceIDs );
+
+	return EditableMesh->MakeUndo();
+}
+
+
+FString FCreateVertexInstancesChange::ToString() const
+{
+	return FString::Printf(
+		TEXT( "Create Vertex Instances [VertexInstancesToCreate:%s]" ),
+		*LogHelpers::ArrayToString( Input.VertexInstancesToCreate ) );
+}
+
+
 TUniquePtr<FChange> FCreateEdgesChange::Execute( UObject* Object )
 {
 	UEditableMesh* EditableMesh = CastChecked<UEditableMesh>( Object );
@@ -88,9 +127,9 @@ TUniquePtr<FChange> FCreatePolygonsChange::Execute( UObject* Object )
 	UEditableMesh* EditableMesh = CastChecked<UEditableMesh>( Object );
 	verify( !EditableMesh->AnyChangesToUndo() );
 	
-	static TArray<FPolygonRef> UnusedNewPolygonRefs;
+	static TArray<FPolygonID> UnusedNewPolygonIDs;
 	static TArray<FEdgeID> UnusedNewEdgeIDs;
-	EditableMesh->CreatePolygons( Input.PolygonsToCreate, /* Out */ UnusedNewPolygonRefs, /* Out */ UnusedNewEdgeIDs );
+	EditableMesh->CreatePolygons( Input.PolygonsToCreate, /* Out */ UnusedNewPolygonIDs, /* Out */ UnusedNewEdgeIDs );
 
 	return EditableMesh->MakeUndo();
 }
@@ -109,7 +148,7 @@ TUniquePtr<FChange> FDeletePolygonsChange::Execute( UObject* Object )
 	UEditableMesh* EditableMesh = CastChecked<UEditableMesh>( Object );
 	verify( !EditableMesh->AnyChangesToUndo() );
 
-	EditableMesh->DeletePolygons( Input.PolygonRefsToDelete, Input.bDeleteOrphanedEdges, Input.bDeleteOrphanedVertices, Input.bDeleteEmptySections );
+	EditableMesh->DeletePolygons( Input.PolygonIDsToDelete, Input.bDeleteOrphanedEdges, Input.bDeleteOrphanedVertices, Input.bDeleteOrphanedVertexInstances, Input.bDeleteEmptySections );
 	return EditableMesh->MakeUndo();
 }
 
@@ -117,8 +156,8 @@ TUniquePtr<FChange> FDeletePolygonsChange::Execute( UObject* Object )
 FString FDeletePolygonsChange::ToString() const
 {
 	return FString::Printf(
-		TEXT( "Delete Polygons [PolygonRefsToDelete:%s, %s, %s]" ),
-		*LogHelpers::ArrayToString( Input.PolygonRefsToDelete ),
+		TEXT( "Delete Polygons [PolygonIDsToDelete:%s, %s, %s]" ),
+		*LogHelpers::ArrayToString( Input.PolygonIDsToDelete ),
 		*LogHelpers::BoolToString( Input.bDeleteOrphanedEdges ),
 		*LogHelpers::BoolToString( Input.bDeleteOrphanedVertices ),
 		*LogHelpers::BoolToString( Input.bDeleteEmptySections ) );
@@ -143,6 +182,24 @@ FString FSetVerticesAttributesChange::ToString() const
 }
 
 
+TUniquePtr<FChange> FSetVertexInstancesAttributesChange::Execute( UObject* Object )
+{
+	UEditableMesh* EditableMesh = CastChecked<UEditableMesh>( Object );
+	verify( !EditableMesh->AnyChangesToUndo() );
+
+	EditableMesh->SetVertexInstancesAttributes( Input.AttributesForVertexInstances );
+	return EditableMesh->MakeUndo();
+}
+
+
+FString FSetVertexInstancesAttributesChange::ToString() const
+{
+	return FString::Printf(
+		TEXT( "Set Vertex Instances Attributes [AttributesForVertexInstances:%s]" ),
+		*LogHelpers::ArrayToString( Input.AttributesForVertexInstances ) );
+}
+
+
 TUniquePtr<FChange> FSetEdgesAttributesChange::Execute( UObject* Object )
 {
 	UEditableMesh* EditableMesh = CastChecked<UEditableMesh>( Object );
@@ -179,6 +236,24 @@ FString FSetPolygonsVertexAttributesChange::ToString() const
 }
 
 
+TUniquePtr<FChange> FChangePolygonsVertexInstancesChange::Execute( UObject* Object )
+{
+	UEditableMesh* EditableMesh = CastChecked<UEditableMesh>( Object );
+	verify( !EditableMesh->AnyChangesToUndo() );
+
+	EditableMesh->ChangePolygonsVertexInstances( Input.VertexInstancesForPolygons );
+	return EditableMesh->MakeUndo();
+}
+
+
+FString FChangePolygonsVertexInstancesChange::ToString() const
+{
+	return FString::Printf(
+		TEXT( "Change Polygons Vertex Instances [VertexInstancesForPolygons:%s]" ),
+		*LogHelpers::ArrayToString( Input.VertexInstancesForPolygons ) );
+}
+
+
 TUniquePtr<FChange> FSetEdgesVerticesChange::Execute( UObject* Object )
 {
 	UEditableMesh* EditableMesh = CastChecked<UEditableMesh>( Object );
@@ -203,7 +278,7 @@ TUniquePtr<FChange> FInsertPolygonPerimeterVerticesChange::Execute( UObject* Obj
 	UEditableMesh* EditableMesh = CastChecked<UEditableMesh>( Object );
 	verify( !EditableMesh->AnyChangesToUndo() );
 
-	EditableMesh->InsertPolygonPerimeterVertices( Input.PolygonRef, Input.InsertBeforeVertexNumber, Input.VerticesToInsert );
+	EditableMesh->InsertPolygonPerimeterVertices( Input.PolygonID, Input.InsertBeforeVertexNumber, Input.VerticesToInsert );
 	return EditableMesh->MakeUndo();
 }
 
@@ -211,8 +286,8 @@ TUniquePtr<FChange> FInsertPolygonPerimeterVerticesChange::Execute( UObject* Obj
 FString FInsertPolygonPerimeterVerticesChange::ToString() const
 {
 	return FString::Printf(
-		TEXT( "Insert Polygon Perimeter Vertices [PolygonRef:%s, InsertBeforeVertexNumber:%lu, VerticesToInsert:%s]" ),
-		*Input.PolygonRef.ToString(),
+		TEXT( "Insert Polygon Perimeter Vertices [PolygonID:%s, InsertBeforeVertexNumber:%lu, VerticesToInsert:%s]" ),
+		*Input.PolygonID.ToString(),
 		Input.InsertBeforeVertexNumber,
 		*LogHelpers::ArrayToString( Input.VerticesToInsert ) );
 }
@@ -223,7 +298,7 @@ TUniquePtr<FChange> FRemovePolygonPerimeterVerticesChange::Execute( UObject* Obj
 	UEditableMesh* EditableMesh = CastChecked<UEditableMesh>( Object );
 	verify( !EditableMesh->AnyChangesToUndo() );
 
-	EditableMesh->RemovePolygonPerimeterVertices( Input.PolygonRef, Input.FirstVertexNumberToRemove, Input.NumVerticesToRemove );
+	EditableMesh->RemovePolygonPerimeterVertices( Input.PolygonID, Input.FirstVertexNumberToRemove, Input.NumVerticesToRemove, Input.bDeleteOrphanedVertexInstances );
 	return EditableMesh->MakeUndo();
 }
 
@@ -231,8 +306,8 @@ TUniquePtr<FChange> FRemovePolygonPerimeterVerticesChange::Execute( UObject* Obj
 FString FRemovePolygonPerimeterVerticesChange::ToString() const
 {
 	return FString::Printf(
-		TEXT( "Remove Polygon Perimeter Vertices [PolygonRef:%s, FirstVertexNumberToRemove:%lu, NumVerticesToRemove:%lu]" ),
-		*Input.PolygonRef.ToString(),
+		TEXT( "Remove Polygon Perimeter Vertices [PolygonID:%s, FirstVertexNumberToRemove:%lu, NumVerticesToRemove:%lu]" ),
+		*Input.PolygonID.ToString(),
 		Input.FirstVertexNumberToRemove,
 		Input.NumVerticesToRemove );
 }
@@ -267,25 +342,6 @@ FString FStartOrEndModificationChange::ToString() const
 }
 
 
-TUniquePtr<FChange> FRetrianglulatePolygonsChange::Execute( UObject* Object )
-{
-	UEditableMesh* EditableMesh = CastChecked<UEditableMesh>( Object );
-	verify( !EditableMesh->AnyChangesToUndo() );
-
-	EditableMesh->RetriangulatePolygons( Input.PolygonRefs, Input.bOnlyOnUndo );
-	return EditableMesh->MakeUndo();
-}
-
-
-FString FRetrianglulatePolygonsChange::ToString() const
-{
-	return FString::Printf(
-		TEXT( "Retriangulate Polygons [PolygonRefs:%s, bOnlyOnUndo:%s]" ),
-		*LogHelpers::ArrayToString( Input.PolygonRefs ),
-		*LogHelpers::BoolToString( Input.bOnlyOnUndo ) );
-}
-
-
 TUniquePtr<FChange> FSetSubdivisionCountChange::Execute( UObject* Object )
 {
 	UEditableMesh* EditableMesh = CastChecked<UEditableMesh>( Object );
@@ -304,41 +360,40 @@ FString FSetSubdivisionCountChange::ToString() const
 }
 
 
-TUniquePtr<FChange> FCreateSectionChange::Execute( UObject* Object )
+TUniquePtr<FChange> FCreatePolygonGroupsChange::Execute( UObject* Object )
 {
 	UEditableMesh* EditableMesh = CastChecked<UEditableMesh>( Object );
 	verify( !EditableMesh->AnyChangesToUndo() );
 
-	EditableMesh->CreateSection( Input.SectionToCreate );
+	static TArray<FPolygonGroupID> UnusedPolygonGroupIDs;
+	EditableMesh->CreatePolygonGroups( Input.PolygonGroupsToCreate, UnusedPolygonGroupIDs );
 	return EditableMesh->MakeUndo();
 }
 
 
-FString FCreateSectionChange::ToString() const
+FString FCreatePolygonGroupsChange::ToString() const
 {
 	return FString::Printf(
-		TEXT( "Create Section [Material:%s, bEnableCollision:%i, bCastShadow:%i]" ),
-		Input.SectionToCreate.Material ? *Input.SectionToCreate.Material->GetName() : TEXT( "<none>" ),
-		Input.SectionToCreate.bEnableCollision,
-		Input.SectionToCreate.bCastShadow );
+		TEXT( "Create Section [PolygonGroupsToCreate:%s]" ),
+		*LogHelpers::ArrayToString( Input.PolygonGroupsToCreate ) );
 }
 
 
-TUniquePtr<FChange> FDeleteSectionChange::Execute( UObject* Object )
+TUniquePtr<FChange> FDeletePolygonGroupsChange::Execute( UObject* Object )
 {
 	UEditableMesh* EditableMesh = CastChecked<UEditableMesh>( Object );
 	verify( !EditableMesh->AnyChangesToUndo() );
 
-	EditableMesh->DeleteSection( Input.SectionID );
+	EditableMesh->DeletePolygonGroups( Input.PolygonGroupIDs );
 	return EditableMesh->MakeUndo();
 }
 
 
-FString FDeleteSectionChange::ToString() const
+FString FDeletePolygonGroupsChange::ToString() const
 {
 	return FString::Printf(
-		TEXT( "Delete Section [SectionID:%i]" ),
-		Input.SectionID.GetValue() );
+		TEXT( "Delete Section [PolygonGroupIDs:%s]" ),
+		*LogHelpers::ArrayToString( Input.PolygonGroupIDs ) );
 }
 
 
diff --git a/Engine/Source/Runtime/MeshEditingRuntime/EditableMeshChanges.h b/Engine/Source/Runtime/MeshEditingRuntime/EditableMeshChanges.h
--- a/Engine/Source/Runtime/MeshEditingRuntime/EditableMeshChanges.h
+++ b/Engine/Source/Runtime/MeshEditingRuntime/EditableMeshChanges.h
@@ -50,6 +50,50 @@ private:
 };
 
 
+struct FDeleteVertexInstancesChangeInput
+{
+	/** The vertex instance IDs to delete */
+	TArray<FVertexInstanceID> VertexInstanceIDsToDelete;
+
+	/** Whether we should also delete any vertices if we delete their only instance */
+	bool bDeleteOrphanedVertices;
+
+	/** Default constructor */
+	FDeleteVertexInstancesChangeInput()
+		: VertexInstanceIDsToDelete(),
+		  bDeleteOrphanedVertices( true )
+	{
+	}
+};
+
+
+class FDeleteVertexInstancesChange : public FChange
+{
+
+public:
+
+	/** Constructor */
+	FDeleteVertexInstancesChange( const FDeleteVertexInstancesChangeInput& InitInput )
+		: Input( InitInput )
+	{
+	}
+
+	FDeleteVertexInstancesChange( FDeleteVertexInstancesChangeInput&& InitInput )
+		: Input( MoveTemp( InitInput ) )
+	{
+	}
+
+	// Parent class overrides
+	virtual TUniquePtr<FChange> Execute( UObject* Object ) override;
+	virtual FString ToString() const override;
+
+private:
+
+	/** The data we need to make this change */
+	FDeleteVertexInstancesChangeInput Input;
+};
+
+
 struct FDeleteEdgesChangeInput
 {
 	/** The edge IDs to get rid of */
@@ -138,6 +182,48 @@ private:
 };
 
 
+struct FCreateVertexInstancesChangeInput
+{
+	/** Information about each vertex instance that will be created */
+	TArray<FVertexInstanceToCreate> VertexInstancesToCreate;
+
+	/** Default constructor */
+	FCreateVertexInstancesChangeInput()
+		: VertexInstancesToCreate()
+	{
+	}
+};
+
+
+class FCreateVertexInstancesChange : public FChange
+{
+
+public:
+
+	/** Constructor */
+	FCreateVertexInstancesChange( const FCreateVertexInstancesChangeInput& InitInput )
+		: Input( InitInput )
+	{
+	}
+
+	FCreateVertexInstancesChange( FCreateVertexInstancesChangeInput&& InitInput )
+		: Input( MoveTemp( InitInput ) )
+	{
+	}
+
+	// Parent class overrides
+	virtual TUniquePtr<FChange> Execute( UObject* Object ) override;
+	virtual FString ToString() const override;
+
+
+private:
+
+	/** The data we need to make this change */
+	FCreateVertexInstancesChangeInput Input;
+
+};
+
+
 struct FCreateEdgesChangeInput
 {
 	/** Information about each edge that will be created */
@@ -225,7 +311,7 @@ private:
 struct FDeletePolygonsChangeInput
 {
 	/** The polygons to get rid of */
-	TArray<FPolygonRef> PolygonRefsToDelete;
+	TArray<FPolygonID> PolygonIDsToDelete;
 
 	/** Whether we should also delete any edges that are left orphaned after deleting this polygon */
 	bool bDeleteOrphanedEdges;
@@ -233,14 +319,18 @@ struct FDeletePolygonsChangeInput
 	/** Whether we should also delete any vertices that are left orphaned after deleting this polygon */
 	bool bDeleteOrphanedVertices;
 
+	/** Whether we should also delete any vertex instances that are left orphaned after deleting this polygon */
+	bool bDeleteOrphanedVertexInstances;
+
 	/** Whether we should also delete any sections that are left empty after deleting this polygon */
 	bool bDeleteEmptySections;
 
 	/** Default constructor */
 	FDeletePolygonsChangeInput()
-		: PolygonRefsToDelete(),
+		: PolygonIDsToDelete(),
 		  bDeleteOrphanedEdges( true ),
 		  bDeleteOrphanedVertices( true ),
+		  bDeleteOrphanedVertexInstances( true ),
 		  bDeleteEmptySections( false )
 	{
 	}
@@ -311,6 +401,41 @@ private:
 };
 
 
+struct FSetVertexInstancesAttributesChangeInput
+{
+	TArray<FAttributesForVertexInstance> AttributesForVertexInstances;
+};
+
+
+class FSetVertexInstancesAttributesChange : public FChange
+{
+
+public:
+
+	/** Constructor */
+	FSetVertexInstancesAttributesChange( const FSetVertexInstancesAttributesChangeInput& InitInput )
+		: Input( InitInput )
+	{
+	}
+
+	FSetVertexInstancesAttributesChange( FSetVertexInstancesAttributesChangeInput&& InitInput )
+		: Input( MoveTemp( InitInput ) )
+	{
+	}
+
+	// Parent class overrides
+	virtual TUniquePtr<FChange> Execute( UObject* Object ) override;
+	virtual FString ToString() const override;
+
+
+private:
+
+	/** The data we need to make this change */
+	FSetVertexInstancesAttributesChangeInput Input;
+
+};
+
+
 struct FSetEdgesAttributesChangeInput
 {
 	/** Which edges we'll be setting attributes for, along with the attribute data to set */
@@ -395,6 +520,48 @@ private:
 };
 
 
+struct FChangePolygonsVertexInstancesChangeInput
+{
+	/** Which polygons we'll be setting vertex instances for, along with the vertex instances to set */
+	TArray<FChangeVertexInstancesForPolygon> VertexInstancesForPolygons;
+
+	/** Default constructor */
+	FChangePolygonsVertexInstancesChangeInput()
+		: VertexInstancesForPolygons()
+	{
+	}
+};
+
+
+class FChangePolygonsVertexInstancesChange : public FChange
+{
+
+public:
+
+	/** Constructor */
+	FChangePolygonsVertexInstancesChange( const FChangePolygonsVertexInstancesChangeInput& InitInput )
+		: Input( InitInput )
+	{
+	}
+
+	FChangePolygonsVertexInstancesChange( FChangePolygonsVertexInstancesChangeInput&& InitInput )
+		: Input( MoveTemp( InitInput ) )
+	{
+	}
+
+	// Parent class overrides
+	virtual TUniquePtr<FChange> Execute( UObject* Object ) override;
+	virtual FString ToString() const override;
+
+
+private:
+
+	/** The data we need to make this change */
+	FChangePolygonsVertexInstancesChangeInput Input;
+
+};
+
+
 struct FSetEdgesVerticesChangeInput
 {
 	/** The edge to set new vertices for */
@@ -440,7 +607,7 @@ private:
 struct FInsertPolygonPerimeterVerticesChangeInput
 {
 	/** The polygon we'll be inserting vertices into */
-	FPolygonRef PolygonRef;
+	FPolygonID PolygonID;
 
 	/** The first polygon perimeter vertex number to insert indices before */
 	uint32 InsertBeforeVertexNumber;
@@ -451,7 +618,7 @@ struct FInsertPolygonPerimeterVerticesChangeInput
 	
 	/** Default constructor */
 	FInsertPolygonPerimeterVerticesChangeInput()
-		: PolygonRef( FSectionID::Invalid, FPolygonID::Invalid ),
+		: PolygonID( FPolygonID::Invalid ),
 		  InsertBeforeVertexNumber( 0 ),
 		  VerticesToInsert()
 	{
@@ -491,7 +658,7 @@ private:
 struct FRemovePolygonPerimeterVerticesChangeInput
 {
 	/** The polygon we'll be removing vertices from */
-	FPolygonRef PolygonRef;
+	FPolygonID PolygonID;
 
 	/** The first polygon perimeter vertex number to remove */
 	uint32 FirstVertexNumberToRemove;
@@ -499,12 +666,16 @@ struct FRemovePolygonPerimeterVerticesChangeInput
 	/** The number of vertices to remove */
 	uint32 NumVerticesToRemove;
 
+	/** Whether orphaned vertex instances should be deleted or not */
+	bool bDeleteOrphanedVertexInstances;
+
 	
 	/** Default constructor */
 	FRemovePolygonPerimeterVerticesChangeInput()
-		: PolygonRef( FSectionID::Invalid, FPolygonID::Invalid ),
+		: PolygonID( FPolygonID::Invalid ),
 		  FirstVertexNumberToRemove( 0 ),
-		  NumVerticesToRemove( 0 )
+		  NumVerticesToRemove( 0 ),
+		  bDeleteOrphanedVertexInstances( false )
 	{
 	}
 };
@@ -580,52 +751,6 @@ private:
 };
 
 
-struct FRetrianglulatePolygonsChangeInput
-{
-	/** The polygons to retriangulate */
-	TArray<FPolygonRef> PolygonRefs;
-
-	/** Whether we should actually do the retriangulation.  This will be inverted every time an Undo/Redo happens */
-	bool bOnlyOnUndo;
-
-	/** Default constructor */
-	FRetrianglulatePolygonsChangeInput()
-		: PolygonRefs(),
-		  bOnlyOnUndo( true )
-	{
-	}
-};
-
-
-class FRetrianglulatePolygonsChange : public FChange
-{
-
-public:
-
-	/** Constructor */
-	FRetrianglulatePolygonsChange( const FRetrianglulatePolygonsChangeInput& InitInput )
-		: Input( InitInput )
-	{
-	}
-
- 	FRetrianglulatePolygonsChange( FRetrianglulatePolygonsChangeInput&& InitInput )
- 		: Input( MoveTemp( InitInput ) )
- 	{
- 	}
-
-	// Parent class overrides
-	virtual TUniquePtr<FChange> Execute( UObject* Object ) override;
-	virtual FString ToString() const override;
-
-
-private:
-
-	/** The data we need to make this change */
-	FRetrianglulatePolygonsChangeInput Input;
-
-};
-
-
 struct FSetSubdivisionCountChangeInput
 {
 	/** The new number of subdivisions to use */
@@ -663,26 +788,26 @@ private:
 };
 
 
-struct FCreateSectionChangeInput
+struct FCreatePolygonGroupsChangeInput
 {
-	/** Information about the section to create */
-	FSectionToCreate SectionToCreate;
+	/** Information about the polygon groups to create */
+	TArray<FPolygonGroupToCreate> PolygonGroupsToCreate;
 
 	/** Default constructor */
-	FCreateSectionChangeInput()
-		: SectionToCreate()
+	FCreatePolygonGroupsChangeInput()
+		: PolygonGroupsToCreate()
 	{
 	}
 };
 
 
-class FCreateSectionChange : public FChange
+class FCreatePolygonGroupsChange : public FChange
 {
 
 public:
 
 	/** Constructor */
-	FCreateSectionChange( const FCreateSectionChangeInput& InitInput )
+	FCreatePolygonGroupsChange( const FCreatePolygonGroupsChangeInput& InitInput )
 		: Input( InitInput )
 	{
 	}
@@ -695,30 +820,30 @@ public:
 private:
 
 	/** The data we need to make this change */
-	FCreateSectionChangeInput Input;
+	FCreatePolygonGroupsChangeInput Input;
 };
 
 
-struct FDeleteSectionChangeInput
+struct FDeletePolygonGroupsChangeInput
 {
-	/** The section to delete */
-	FSectionID SectionID;
+	/** The polygon group IDs to delete */
+	TArray<FPolygonGroupID> PolygonGroupIDs;
 
 	/** Default constructor */
-	FDeleteSectionChangeInput()
-		: SectionID( FSectionID::Invalid )
+	FDeletePolygonGroupsChangeInput()
+		: PolygonGroupIDs()
 	{
 	}
 };
 
 
-class FDeleteSectionChange : public FChange
+class FDeletePolygonGroupsChange : public FChange
 {
 
 public:
 
 	/** Constructor */
-	FDeleteSectionChange( const FDeleteSectionChangeInput& InitInput )
+	FDeletePolygonGroupsChange( const FDeletePolygonGroupsChangeInput& InitInput )
 		: Input( InitInput )
 	{
 	}
@@ -731,5 +856,5 @@ public:
 private:
 
 	/** The data we need to make this change */
-	FDeleteSectionChangeInput Input;
+	FDeletePolygonGroupsChangeInput Input;
 };
diff --git a/Engine/Source/Runtime/MeshEditingRuntime/EditableStaticMeshAdapter.cpp b/Engine/Source/Runtime/MeshEditingRuntime/EditableStaticMeshAdapter.cpp
new file mode 100644
--- /dev/null
+++ b/Engine/Source/Runtime/MeshEditingRuntime/EditableStaticMeshAdapter.cpp
@@ -0,0 +1,1733 @@
+// Copyright 1998-2015 Epic Games, Inc. All Rights Reserved.
+
+#include "EditableStaticMeshAdapter.h"
+#include "EditableMesh.h"
+#include "EditableMeshChanges.h"
+#include "Engine/StaticMesh.h"
+#include "Components/StaticMeshComponent.h"
+#include "StaticMeshResources.h"
+#include "PhysicsEngine/BodySetup.h"	// For collision generation
+#include "ProfilingDebugging/ScopedTimers.h"	// For FAutoScopedDurationTimer
+#include "EditableMeshFactory.h"
+
+
+const FTriangleID FTriangleID::Invalid( TNumericLimits<uint32>::Max() );
+
+
+UEditableStaticMeshAdapter::UEditableStaticMeshAdapter()
+	: StaticMesh( nullptr ),
+	  RecreateRenderStateContext(),
+	  StaticMeshLODIndex( 0 )
+{
+}
+
+
+inline void UEditableStaticMeshAdapter::EnsureIndexBufferIs32Bit()
+{
+	FStaticMeshLODResources& StaticMeshLOD = GetStaticMeshLOD();
+	if( !StaticMeshLOD.IndexBuffer.Is32Bit() )
+	{
+		// Need a 32-bit index buffer
+		static TArray<uint32> AllIndices;
+		StaticMeshLOD.IndexBuffer.GetCopy( /* Out */ AllIndices );
+		StaticMeshLOD.IndexBuffer.SetIndices( AllIndices, EIndexBufferStride::Force32Bit );
+	}
+}
+
+
+inline void UEditableStaticMeshAdapter::UpdateIndexBufferFormatIfNeeded( const TArray<FMeshTriangle>& Triangles )
+{
+	const FStaticMeshLODResources& StaticMeshLOD = GetStaticMeshLOD();
+	if( !StaticMeshLOD.IndexBuffer.Is32Bit() )
+	{
+		for( const FMeshTriangle& Triangle : Triangles )
+		{
+			for( const FVertexInstanceID VertexInstanceID : Triangle.VertexInstanceIDs )
+			{
+				if( VertexInstanceID.GetValue() > TNumericLimits<uint16>::Max() )
+				{
+					EnsureIndexBufferIs32Bit();
+					return;
+				}
+			}
+		}
+	}
+}
+
+
+void UEditableStaticMeshAdapter::InitEditableStaticMesh( UEditableMesh* EditableMesh, class UPrimitiveComponent& Component, const FEditableMeshSubMeshAddress& InitSubMeshAddress )
+{
+	EditableMesh->SetSubMeshAddress( InitSubMeshAddress );
+	StaticMeshLODIndex = InitSubMeshAddress.LODIndex;
+
+	// We're partial to static mesh components, here
+	UStaticMeshComponent* StaticMeshComponent = Cast<UStaticMeshComponent>( &Component );
+	if( StaticMeshComponent != nullptr )
+	{
+		UStaticMesh* ComponentStaticMesh = StaticMeshComponent->GetStaticMesh();
+		if( ComponentStaticMesh != nullptr && ComponentStaticMesh->HasValidRenderData() )
+		{
+			this->StaticMesh = ComponentStaticMesh;
+			this->OriginalStaticMesh = ComponentStaticMesh;
+
+			const FStaticMeshRenderData& StaticMeshRenderData = *StaticMesh->RenderData;
+			if( StaticMeshLODIndex >= 0 && StaticMeshLODIndex < StaticMeshRenderData.LODResources.Num() )
+			{
+				{
+					// @todo mesheditor urgent: Currently, we're disabling many of the optimized index buffers that were precomputed
+					// for static meshes when they become editable.  This is just so that we don't have to keep this data up to
+					// date as we perform live edits to the geometry.  Later, we should probably get this updated as we go, or 
+					// lazily update the buffers when committing a final change or saving.  Without clearing these values, some
+					// graphical artifacts will be visible while editing the mesh (flickering shadows, for example.)
+					FStaticMeshLODResources& StaticMeshLOD = StaticMesh->RenderData->LODResources[ StaticMeshLODIndex ];
+					StaticMeshLOD.bHasAdjacencyInfo = false;
+					StaticMeshLOD.bHasDepthOnlyIndices = false;
+					StaticMeshLOD.bHasReversedIndices = false;
+					StaticMeshLOD.bHasReversedDepthOnlyIndices = false;
+					StaticMeshLOD.DepthOnlyNumTriangles = 0;
+				}
+
+				const FStaticMeshLODResources& StaticMeshLOD = StaticMeshRenderData.LODResources[ StaticMeshLODIndex ];
+
+				// Store off the number of texture coordinates in this mesh
+				EditableMesh->TextureCoordinateCount = StaticMeshLOD.GetNumTexCoords();
+
+				// Vertices
+				const int32 NumRenderingVertices = StaticMeshLOD.PositionVertexBuffer.GetNumVertices();
+				const int32 NumUVs = StaticMeshLOD.GetNumTexCoords();
+				const bool bHasColor = StaticMeshLOD.ColorVertexBuffer.GetNumVertices() > 0;
+				check( !bHasColor || StaticMeshLOD.ColorVertexBuffer.GetNumVertices() == StaticMeshLOD.VertexBuffer.GetNumVertices() );
+
+				// @todo mesheditor cleanup: This code is very similar to the static mesh build code; try to share helper structs
+				static TMultiMap<int32, int32> OverlappingRenderingVertexIndices;
+				{
+					OverlappingRenderingVertexIndices.Reset();
+
+					/** Helper struct for building acceleration structures. */
+					struct FIndexAndZ
+					{
+						float Z;
+						int32 Index;
+
+						FIndexAndZ() 
+						{
+						}
+
+						/** Initialization constructor. */
+						FIndexAndZ(int32 InIndex, FVector V)
+						{
+							Z = 0.30f * V.X + 0.33f * V.Y + 0.37f * V.Z;
+							Index = InIndex;
+						}
+					};
+
+					// Build a temporary array of vertex instance indices, sorted by their Z value.  This will accelerate
+					// searching through to find duplicates
+					static TArray<FIndexAndZ> RenderingVertexIndicesSortedByZ;
+					{
+						RenderingVertexIndicesSortedByZ.SetNumUninitialized( NumRenderingVertices, false );
+						for( int32 RenderingVertexIndex = 0; RenderingVertexIndex < NumRenderingVertices; ++RenderingVertexIndex )
+						{
+							const FVector VertexPosition = StaticMeshLOD.PositionVertexBuffer.VertexPosition( RenderingVertexIndex );
+							RenderingVertexIndicesSortedByZ[ RenderingVertexIndex ] = FIndexAndZ( RenderingVertexIndex, VertexPosition );
+						}
+
+						// Sort the vertices by z value
+						struct FCompareIndexAndZ
+						{
+							FORCEINLINE bool operator()( FIndexAndZ const& A, FIndexAndZ const& B ) const { return A.Z < B.Z; }
+						};
+						RenderingVertexIndicesSortedByZ.Sort( FCompareIndexAndZ() );
+					}
+
+					// Search for duplicates, quickly!
+					const float ComparisonThreshold = KINDA_SMALL_NUMBER;	// @todo mesheditor: Tweak "weld" threshold
+					for( int32 RenderingVertexIterA = 0; RenderingVertexIterA < NumRenderingVertices; ++RenderingVertexIterA )
+					{
+						// only need to search forward, since we add pairs both ways
+						for( int32 RenderingVertexIterB = RenderingVertexIterA + 1; RenderingVertexIterB < NumRenderingVertices; ++RenderingVertexIterB )
+						{
+							if( FMath::Abs( RenderingVertexIndicesSortedByZ[ RenderingVertexIterB ].Z - RenderingVertexIndicesSortedByZ[ RenderingVertexIterA ].Z ) > ComparisonThreshold )
+							{
+								break; // can't be any more dups
+							}
+
+							const int32 RenderingVertexIndexA = RenderingVertexIndicesSortedByZ[ RenderingVertexIterA ].Index;
+							const int32 RenderingVertexIndexB = RenderingVertexIndicesSortedByZ[ RenderingVertexIterB ].Index;
+
+							const FVector VertexPositionA = StaticMeshLOD.PositionVertexBuffer.VertexPosition( RenderingVertexIndexA );
+							const FVector VertexPositionB = StaticMeshLOD.PositionVertexBuffer.VertexPosition( RenderingVertexIndexB );
+
+							if( VertexPositionA.Equals( VertexPositionB, ComparisonThreshold ) )
+							{
+								OverlappingRenderingVertexIndices.Add( RenderingVertexIndexA, RenderingVertexIndexB );
+								OverlappingRenderingVertexIndices.Add( RenderingVertexIndexB, RenderingVertexIndexA );
+							}
+						}
+					}
+				}
+
+				// We'll now make sure we have an editable mesh vertex created for every uniquely-positioned vertex instance.
+				// Note that it's important that we process all vertices, not only the vertices that are referenced by triangles
+				// in the index buffer, because we properly support meshes with vertices that are not yet connected to any
+				// polygons.  These vertices will simply not have editable mesh polygons or edges connected to them, but will
+				// still be interactable in the editor.
+				for( int32 RenderingVertexIndex = 0; RenderingVertexIndex < NumRenderingVertices; ++RenderingVertexIndex )
+				{
+					const FVector VertexPosition = StaticMeshLOD.PositionVertexBuffer.VertexPosition( RenderingVertexIndex );
+
+					// Check to see if we already have this vertex
+					bool bAlreadyHaveVertexForPosition = false;
+					{
+						static TArray<int32> ThisRenderingVertexOverlaps;
+						ThisRenderingVertexOverlaps.Reset();
+						OverlappingRenderingVertexIndices.MultiFind( RenderingVertexIndex, /* Out */ ThisRenderingVertexOverlaps );
+
+						for( int32 OverlappingRenderingVertexIter = 0; OverlappingRenderingVertexIter < ThisRenderingVertexOverlaps.Num(); ++OverlappingRenderingVertexIter )
+						{
+							const int32 OverlappingRenderingVertexIndex = ThisRenderingVertexOverlaps[ OverlappingRenderingVertexIter ];
+
+							// If the overlapping vertex instance index is smaller than our current index, we can safely assume that
+							// we've already processed this vertex position and created an editable mesh vertex for it.
+							if( OverlappingRenderingVertexIndex < RenderingVertexIndex )
+							{
+								check( EditableMesh->VertexInstances.IsAllocated( OverlappingRenderingVertexIndex ) );
+								const FVertexID ExistingVertexID = EditableMesh->VertexInstances[ OverlappingRenderingVertexIndex ].VertexID;
+								FMeshVertex& ExistingVertex = EditableMesh->Vertices[ ExistingVertexID.GetValue() ];
+
+								// We already have a unique editable vertex for this vertex instance position, so link them!
+								EditableMesh->VertexInstances.Insert( RenderingVertexIndex, FMeshVertexInstance() );
+								EditableMesh->VertexInstances[ RenderingVertexIndex ].VertexID = ExistingVertexID;
+
+								const FVertexInstanceID VertexInstanceID( RenderingVertexIndex );
+								checkSlow( !ExistingVertex.VertexInstanceIDs.Contains( VertexInstanceID ) );
+								ExistingVertex.VertexInstanceIDs.Add( VertexInstanceID );
+								bAlreadyHaveVertexForPosition = true;
+
+								break;
+							}
+						}
+					}
+
+					if( !bAlreadyHaveVertexForPosition )
+					{
+						const FVertexID NewVertexID( EditableMesh->Vertices.Add( FMeshVertex() ) );
+						FMeshVertex& NewVertex = EditableMesh->Vertices[ NewVertexID.GetValue() ];
+						NewVertex.VertexPosition = VertexPosition;
+						NewVertex.CornerSharpness = 0.0f;
+
+						EditableMesh->VertexInstances.Insert( RenderingVertexIndex, FMeshVertexInstance() );
+						EditableMesh->VertexInstances[ RenderingVertexIndex ].VertexID = NewVertexID;
+						
+						// @todo mesheditor: If a mesh somehow contained vertex instances that no triangle was referencing, this would cause
+						// the vertex instance to be ignored by the editable mesh code.  It would just sit in the vertex buffer (and in the
+						// editable mesh vertex's RenderingVertexIndices list), but would never be touched.  The editable mesh code only
+						// creates vertex instances for vertices that are attached to polygons, so this should never happen with meshes
+						// that we create and save.  Only if the incoming data had orphan vertices in it.  Should hopefully not be a problem.
+						const FVertexInstanceID VertexInstanceID( RenderingVertexIndex );
+						NewVertex.VertexInstanceIDs.Add( VertexInstanceID );
+
+						// NOTE: The new vertex's connected polygons will be filled in down below, as we're processing mesh triangles
+					}
+
+					// Populate the vertex instance attributes
+					{
+						FMeshVertexInstance& VertexInstance = EditableMesh->VertexInstances[ RenderingVertexIndex ];
+						VertexInstance.VertexUVs.Reserve( NumUVs );
+						for( int32 UVIndex = 0; UVIndex < NumUVs; ++UVIndex )
+						{
+							VertexInstance.VertexUVs.Add( StaticMeshLOD.VertexBuffer.GetVertexUV( RenderingVertexIndex, UVIndex ) );
+						}
+
+						const FVector Normal = StaticMeshLOD.VertexBuffer.VertexTangentZ( RenderingVertexIndex );
+						const FVector Tangent = StaticMeshLOD.VertexBuffer.VertexTangentX( RenderingVertexIndex );
+						const FVector Binormal = StaticMeshLOD.VertexBuffer.VertexTangentY( RenderingVertexIndex );
+						VertexInstance.Normal = Normal;
+						VertexInstance.Tangent = Tangent;
+						VertexInstance.BinormalSign = GetBasisDeterminantSign(Tangent, Binormal, Normal);
+
+						VertexInstance.Color = bHasColor ? FLinearColor( StaticMeshLOD.ColorVertexBuffer.VertexColor( RenderingVertexIndex ) ) : FLinearColor::White;
+					}
+
+				}
+
+
+				const FIndexArrayView RenderingIndices = StaticMeshLOD.IndexBuffer.GetArrayView();
+
+
+				static TMap<uint64, FEdgeID> UniqueEdgeToEdgeID;
+				UniqueEdgeToEdgeID.Reset();
+
+				static TMultiMap<FEdgeID, TTuple<FVertexInstanceID, FVertexInstanceID>> EdgeToVertexInstancePair;
+				EdgeToVertexInstancePair.Reset();
+
+				// Add all polygon groups from the static mesh sections
+				const uint32 NumSections = StaticMeshLOD.Sections.Num();
+				for( uint32 RenderingSectionIndex = 0; RenderingSectionIndex < NumSections; ++RenderingSectionIndex )
+				{
+					const FStaticMeshSection& RenderingSection = StaticMeshLOD.Sections[ RenderingSectionIndex ];
+
+					// Create a new polygon group
+					const FPolygonGroupID NewPolygonGroupID( EditableMesh->PolygonGroups.Add( FMeshPolygonGroup() ) );
+					FMeshPolygonGroup& NewPolygonGroup = EditableMesh->PolygonGroups[ NewPolygonGroupID.GetValue() ];
+
+					NewPolygonGroup.Material = StaticMesh->GetMaterial( RenderingSection.MaterialIndex );
+					NewPolygonGroup.bEnableCollision = RenderingSection.bEnableCollision;
+					NewPolygonGroup.bCastShadow = RenderingSection.bCastShadow;
+
+					// Create a rendering polygon group for holding the triangulated data and references to the static mesh rendering section.
+					// This is indexed by the same FPolygonGroupID as the PolygonGroups.
+					RenderingPolygonGroups.Insert( NewPolygonGroupID.GetValue(), FRenderingPolygonGroup() );
+					FRenderingPolygonGroup& NewRenderingPolygonGroup = RenderingPolygonGroups[ NewPolygonGroupID.GetValue() ];
+
+					const uint32 NumSectionTriangles = RenderingSection.NumTriangles;
+					NewRenderingPolygonGroup.Triangles.Reserve( NumSectionTriangles );
+					NewRenderingPolygonGroup.MaxTriangles = NumSectionTriangles;
+					NewRenderingPolygonGroup.RenderingSectionIndex = RenderingSectionIndex;
+
+					for( uint32 SectionTriangleIndex = 0; SectionTriangleIndex < NumSectionTriangles; ++SectionTriangleIndex )
+					{
+						const uint32 RenderingTriangleFirstVertexIndex = FRenderingPolygonGroup::TriangleIndexToRenderingTriangleFirstIndex( RenderingSection, FTriangleID( SectionTriangleIndex ) );
+
+						uint32 TriangleRenderingVertexIndices[ 3 ];
+						FVertexID TriangleVertexIDs[ 3 ];
+						for( uint32 TriangleVertexIndex = 0; TriangleVertexIndex < 3; ++TriangleVertexIndex )
+						{
+							TriangleRenderingVertexIndices[ TriangleVertexIndex ] = RenderingIndices[ RenderingTriangleFirstVertexIndex + TriangleVertexIndex ];
+							TriangleVertexIDs[ TriangleVertexIndex ] = EditableMesh->VertexInstances[ TriangleRenderingVertexIndices[ TriangleVertexIndex ] ].VertexID;
+						}
+
+						// Make sure we have a valid triangle.  The triangle can be invalid because at least two if it's vertex indices 
+						// point to the exact same vertex.  The triangle is degenerate.  This can happen due to us welding the overlapping 
+						// vertices because they were either extremely close to each other (or exactly overlapping.)  We'll ignore this triangle.
+						const bool bIsValidTriangle =
+							TriangleVertexIDs[ 0 ] != TriangleVertexIDs[ 1 ] &&
+							TriangleVertexIDs[ 1 ] != TriangleVertexIDs[ 2 ] &&
+							TriangleVertexIDs[ 2 ] != TriangleVertexIDs[ 0 ];
+						if( bIsValidTriangle )
+						{
+							// Static meshes only support triangles, so there's no need to triangulate anything yet.  We'll make both
+							// a triangle and a polygon here.
+							const int32 NewTriangleIndex = SectionTriangleIndex;
+
+							NewRenderingPolygonGroup.Triangles.InsertUninitialized( NewTriangleIndex );
+							FMeshTriangle& NewTriangle = NewRenderingPolygonGroup.Triangles[ NewTriangleIndex ];
+
+							// Insert a polygon into the mesh
+							const FPolygonID NewPolygonID( EditableMesh->Polygons.Add( FMeshPolygon() ) );
+							FMeshPolygon& NewPolygon = EditableMesh->Polygons[ NewPolygonID.GetValue() ];
+							NewPolygon.PolygonGroupID = NewPolygonGroupID;
+							NewPolygonGroup.Polygons.Add( NewPolygonID );
+
+							// Create a rendering polygon mirror, indexed by the same ID
+							RenderingPolygons.Insert( NewPolygonID.GetValue(), FRenderingPolygon() );
+							FRenderingPolygon& NewRenderingPolygon = RenderingPolygons[ NewPolygonID.GetValue() ];
+							NewRenderingPolygon.TriangulatedPolygonTriangleIndices.Add( FTriangleID( NewTriangleIndex ) );
+
+
+							// Static meshes don't support polygons with holes, so we always start out with only a perimeter contour per polygon
+							FMeshPolygonContour& PerimeterContour = NewPolygon.PerimeterContour;
+							PerimeterContour.VertexInstanceIDs.Reserve( 3 );
+
+							// Connect vertices
+							for( uint32 TriangleVertexIndex = 0; TriangleVertexIndex < 3; ++TriangleVertexIndex )
+							{
+								const uint32 RenderingVertexIndex = TriangleRenderingVertexIndices[ TriangleVertexIndex ];
+								const FVertexID VertexID = TriangleVertexIDs[ TriangleVertexIndex ];
+								const FVertexInstanceID VertexInstanceID = FVertexInstanceID( RenderingVertexIndex );
+
+								NewPolygon.PerimeterContour.VertexInstanceIDs.Add( VertexInstanceID );
+
+								FMeshVertexInstance& VertexInstance = EditableMesh->VertexInstances[ VertexInstanceID.GetValue() ];
+								VertexInstance.ConnectedPolygons.Add( NewPolygonID );
+
+								// The triangle points to each of its three vertices
+								NewTriangle.VertexInstanceIDs[ TriangleVertexIndex ] = VertexInstanceID;
+							}
+
+							// Add triangle to polygon triangulation array
+							NewPolygon.Triangles.Add( NewTriangle );
+
+							// Connect edges
+							{
+								struct Local
+								{
+									inline static uint64 Make64BitValueForEdge( const FVertexID EdgeVertexID0, const FVertexID EdgeVertexID1 )
+									{
+										return uint64( ( uint64( EdgeVertexID0.GetValue() ) << 32 ) | uint64( EdgeVertexID1.GetValue() ) );
+									};
+								};
+
+
+								// Add the edges of this triangle
+								for( uint32 TriangleEdgeNumber = 0; TriangleEdgeNumber < 3; ++TriangleEdgeNumber )
+								{
+									uint32 EdgeRenderingVertexIndices[ 2 ];
+									EdgeRenderingVertexIndices[ 0 ] = RenderingIndices[ RenderingTriangleFirstVertexIndex + ( TriangleEdgeNumber + 0 ) % 3 ];
+									EdgeRenderingVertexIndices[ 1 ] = RenderingIndices[ RenderingTriangleFirstVertexIndex + ( TriangleEdgeNumber + 1 ) % 3 ];
+
+									FVertexID EdgeVertexIDs[ 2 ];
+									EdgeVertexIDs[ 0 ] = EditableMesh->VertexInstances[ EdgeRenderingVertexIndices[ 0 ] ].VertexID;
+									EdgeVertexIDs[ 1 ] = EditableMesh->VertexInstances[ EdgeRenderingVertexIndices[ 1 ] ].VertexID;
+
+									FVertexInstanceID EdgeVertexInstanceIDs[ 2 ];
+
+									// Check to see if this edge already exists
+									bool bAlreadyHaveEdge = false;
+									FEdgeID EdgeID = FEdgeID::Invalid;
+									{
+										FEdgeID* FoundEdgeIDPtr = UniqueEdgeToEdgeID.Find( Local::Make64BitValueForEdge( EdgeVertexIDs[ 0 ], EdgeVertexIDs[ 1 ] ) );
+										if( FoundEdgeIDPtr != nullptr )
+										{
+											EdgeVertexInstanceIDs[ 0 ] = FVertexInstanceID( EdgeRenderingVertexIndices[ 0 ] );
+											EdgeVertexInstanceIDs[ 1 ] = FVertexInstanceID( EdgeRenderingVertexIndices[ 1 ] );
+											bAlreadyHaveEdge = true;
+											EdgeID = *FoundEdgeIDPtr;
+										}
+										else
+										{
+											// Try the other way around
+											FoundEdgeIDPtr = UniqueEdgeToEdgeID.Find( Local::Make64BitValueForEdge( EdgeVertexIDs[ 1 ], EdgeVertexIDs[ 0 ] ) );
+
+											if( FoundEdgeIDPtr != nullptr )
+											{
+												EdgeVertexInstanceIDs[ 0 ] = FVertexInstanceID( EdgeRenderingVertexIndices[ 1 ] );
+												EdgeVertexInstanceIDs[ 1 ] = FVertexInstanceID( EdgeRenderingVertexIndices[ 0 ] );
+												bAlreadyHaveEdge = true;
+												EdgeID = *FoundEdgeIDPtr;
+											}
+										}
+									}
+
+									if( !bAlreadyHaveEdge )
+									{
+										// Create the new edge.  We'll connect it to its polygons later on.
+										EdgeID = FEdgeID( EditableMesh->Edges.Add( FMeshEdge() ) );
+										FMeshEdge& NewEdge = EditableMesh->Edges[ EdgeID.GetValue() ];
+
+										NewEdge.VertexIDs[ 0 ] = EdgeVertexIDs[ 0 ];
+										NewEdge.VertexIDs[ 1 ] = EdgeVertexIDs[ 1 ];
+										NewEdge.bIsHardEdge = false;
+										NewEdge.CreaseSharpness = 0.0f;
+
+										UniqueEdgeToEdgeID.Add( Local::Make64BitValueForEdge( EdgeVertexIDs[ 0 ], EdgeVertexIDs[ 1 ] ), EdgeID );
+
+										EdgeVertexInstanceIDs[ 0 ] = FVertexInstanceID( EdgeRenderingVertexIndices[ 0 ] );
+										EdgeVertexInstanceIDs[ 1 ] = FVertexInstanceID( EdgeRenderingVertexIndices[ 1 ] );
+									}
+
+									// Each edge will point back to the polygon that its connected to.  Remember, an edge can be shared by multiple
+									// polygons, but usually its best if only shared by up to two.
+									FMeshEdge& Edge = EditableMesh->Edges[ EdgeID.GetValue() ];
+
+									Edge.ConnectedPolygons.AddUnique( NewPolygonID );
+
+									// Connect the end vertices to the edge
+									EditableMesh->Vertices[ EdgeVertexIDs[ 0 ].GetValue() ].ConnectedEdgeIDs.AddUnique( EdgeID );
+									EditableMesh->Vertices[ EdgeVertexIDs[ 1 ].GetValue() ].ConnectedEdgeIDs.AddUnique( EdgeID );
+
+									// Determine whether the edge is hard by checking whether
+									if( !Edge.bIsHardEdge )
+									{
+										TArray<TTuple<FVertexInstanceID, FVertexInstanceID>> VertexInstancePairs;
+										EdgeToVertexInstancePair.MultiFind( EdgeID, VertexInstancePairs );
+
+										for( const TTuple<FVertexInstanceID, FVertexInstanceID>& VertexInstancePair : VertexInstancePairs )
+										{
+											// If either of the ends of the edge just added have different normals to any overlapping edge, mark it as a hard edge
+											if( EditableMesh->VertexInstances[ VertexInstancePair.Get<0>().GetValue() ].Normal != EditableMesh->VertexInstances[ EdgeVertexInstanceIDs[ 0 ].GetValue() ].Normal ||
+												EditableMesh->VertexInstances[ VertexInstancePair.Get<1>().GetValue() ].Normal != EditableMesh->VertexInstances[ EdgeVertexInstanceIDs[ 1 ].GetValue() ].Normal )
+											{
+												Edge.bIsHardEdge = true;
+												break;
+											}
+										}
+
+										EdgeToVertexInstancePair.Add( EdgeID, MakeTuple( EdgeVertexInstanceIDs[ 0 ], EdgeVertexInstanceIDs[ 1 ] ) );
+									}
+								}
+							}
+						}
+						else
+						{
+							// Triangle was not valid.  This will result in an empty entry in our Triangles sparse array.  Luckily,
+							// the triangle is already degenerate so we don't need to change anything.  This triangle index will be
+							// re-used if a new triangle needs to be created during editing
+							// @todo mesheditor: This can cause vertex instances to be orphaned.  Should we delete them?
+						}
+					}
+				}
+			}
+		}
+	}
+
+	EditableMesh->RefreshOpenSubdiv();
+
+	// Cache polygon tangent bases
+	static TArray<FPolygonID> PolygonIDs;
+	PolygonIDs.Reset();
+	for( int32 PolygonIndex = 0; PolygonIndex < EditableMesh->Polygons.GetMaxIndex(); ++PolygonIndex )
+	{
+		if( EditableMesh->Polygons.IsAllocated( PolygonIndex ) )
+		{
+			PolygonIDs.Add( FPolygonID( PolygonIndex ) );
+		}
+	}
+	EditableMesh->GeneratePolygonTangentsAndNormals( PolygonIDs );
+
+#if 0
+	// Test tangent generation
+	static TArray<FPolygonID> AllPolygons;
+	AllPolygons.Reset();
+	for( int32 PolygonIndex = 0; PolygonIndex < EditableMesh->Polygons.GetMaxIndex(); ++PolygonIndex )
+	{
+		if( EditableMesh->Polygons.IsAllocated( PolygonIndex ) )
+		{
+			EditableMesh->PolygonsPendingNewTangentBasis.Add( FPolygonID( PolygonIndex ) );
+		}
+	}
+
+	EditableMesh->GenerateTangentsAndNormals();
+#endif
+}
+
+
+void UEditableStaticMeshAdapter::InitFromBlankStaticMesh( UEditableMesh* EditableMesh, UStaticMesh& InStaticMesh )
+{
+	StaticMesh = &InStaticMesh;
+}
+
+int32 GetStaticMeshMaterialIndex( const UStaticMesh* StaticMesh, const UMaterialInterface* MaterialInterface )
+{
+	return StaticMesh->StaticMaterials.IndexOfByPredicate(
+		[ MaterialInterface ]( const FStaticMaterial& StaticMaterial ) { return StaticMaterial.MaterialInterface == MaterialInterface; }
+	);
+}
+
+
+void UEditableStaticMeshAdapter::OnRebuildRenderMesh( const UEditableMesh* EditableMesh )
+{
+	// @todo mesheditor urgent subdiv: Saw some editable mesh corruption artifacts when testing subDs in VR
+
+	check( RecreateRenderStateContext.IsValid() );
+
+	FStaticMeshLODResources& StaticMeshLOD = GetStaticMeshLOD();
+
+	// Build new vertex buffers
+	static TArray< FStaticMeshBuildVertex > StaticMeshBuildVertices;
+	StaticMeshBuildVertices.Reset();
+
+	static TArray< uint32 > IndexBuffer;
+	IndexBuffer.Reset();
+
+	StaticMeshLOD.Sections.Empty( EditableMesh->PolygonGroups.Num() );
+
+	bool bHasColor = false;
+
+	if( EditableMesh->IsPreviewingSubdivisions() )
+	{
+		check( EditableMesh->GetSubdivisionCount() > 0 );
+
+		const int32 SectionCount = EditableMesh->SubdivisionLimitData.Sections.Num();
+
+		// @todo mesheditor subdiv: Only 2 UVs supported for now, just to avoid having to use a dynamic array per vertex; needs a new data layout, probably (SoA)
+		const int32 SubdivMeshTextureCoordinateCount = FMath::Min( EditableMesh->GetTextureCoordinateCount(), 2 );
+
+		// The Sections sparse array mirrors the SubdivisionLimitData sections array;
+		// iterate through it in parallel in order to get the material index and other section properties
+		TSparseArray< FMeshPolygonGroup >::TConstIterator PolygonGroupIt( EditableMesh->PolygonGroups );
+		check( EditableMesh->PolygonGroups.Num() == SectionCount );
+
+		for( int32 SectionNumber = 0; SectionNumber < SectionCount; ++SectionNumber )
+		{
+			const FMeshPolygonGroup& PolygonGroup = *PolygonGroupIt;
+			const FSubdivisionLimitSection& SubdivisionSection = EditableMesh->SubdivisionLimitData.Sections[ SectionNumber ];
+
+			const int32 SectionTriangleCount = SubdivisionSection.SubdividedQuads.Num() * 2;
+
+			// @todo mesheditor subdiv perf: Ideally, if no topology changed we can just fill vertex data and not touch index buffers
+			const int32 FirstSectionVertexIndex = StaticMeshBuildVertices.Num();
+			StaticMeshBuildVertices.AddUninitialized( SectionTriangleCount * 3 );
+
+			const int32 FirstIndexInSection = IndexBuffer.Num();
+			IndexBuffer.Reserve( IndexBuffer.Num() + SectionTriangleCount * 3 );
+
+			// Create new rendering section
+			StaticMeshLOD.Sections.Add( FStaticMeshSection() );
+			FStaticMeshSection& StaticMeshSection = StaticMeshLOD.Sections.Last();
+
+			StaticMeshSection.FirstIndex = FirstIndexInSection;
+			StaticMeshSection.NumTriangles = SectionTriangleCount;
+			StaticMeshSection.MinVertexIndex = FirstSectionVertexIndex;
+			StaticMeshSection.MaxVertexIndex = FirstSectionVertexIndex + SectionTriangleCount * 3;
+
+			const int32 MaterialIndex = GetStaticMeshMaterialIndex( StaticMesh, PolygonGroup.Material );
+			check( MaterialIndex != INDEX_NONE );
+			StaticMeshSection.MaterialIndex = MaterialIndex;
+			StaticMeshSection.bEnableCollision = PolygonGroup.bEnableCollision;
+			StaticMeshSection.bCastShadow = PolygonGroup.bCastShadow;
+
+			// Fill vertices
+			int32 NextVertexIndex = FirstSectionVertexIndex;
+			for( int32 QuadNumber = 0; QuadNumber < SubdivisionSection.SubdividedQuads.Num(); ++QuadNumber )
+			{
+				const FSubdividedQuad& SubdividedQuad = SubdivisionSection.SubdividedQuads[ QuadNumber ];
+
+				// @todo mesheditor subdiv debug
+// 				GWarn->Logf( TEXT( "Q%i V%i: U:%0.2f, V:%0.2f" ), QuadNumber, 0, SubdividedQuad.QuadVertex0.TextureCoordinate0.X, SubdividedQuad.QuadVertex0.TextureCoordinate0.Y );
+// 				GWarn->Logf( TEXT( "Q%i V%i: U:%0.2f, V:%0.2f" ), QuadNumber, 1, SubdividedQuad.QuadVertex1.TextureCoordinate0.X, SubdividedQuad.QuadVertex1.TextureCoordinate0.Y );
+// 				GWarn->Logf( TEXT( "Q%i V%i: U:%0.2f, V:%0.2f" ), QuadNumber, 2, SubdividedQuad.QuadVertex2.TextureCoordinate0.X, SubdividedQuad.QuadVertex2.TextureCoordinate0.Y );
+// 				GWarn->Logf( TEXT( "Q%i V%i: U:%0.2f, V:%0.2f" ), QuadNumber, 3, SubdividedQuad.QuadVertex3.TextureCoordinate0.X, SubdividedQuad.QuadVertex3.TextureCoordinate0.Y );
+
+				for( int32 TriangleNumber = 0; TriangleNumber < 2; ++TriangleNumber )
+				{
+					for( int32 TriangleVertexNumber = 0; TriangleVertexNumber < 3; ++TriangleVertexNumber )
+					{
+						int32 QuadVertexNumber;
+						if( TriangleNumber == 0 )
+						{
+							QuadVertexNumber = ( TriangleVertexNumber == 0 ) ? 0 : ( TriangleVertexNumber == 1 ? 2 : 1 );
+						}
+						else
+						{
+							QuadVertexNumber = ( TriangleVertexNumber == 0 ) ? 0 : ( TriangleVertexNumber == 1 ? 3 : 2 );
+						}
+
+						const FSubdividedQuadVertex& QuadVertex = SubdividedQuad.GetQuadVertex( QuadVertexNumber );
+
+						const FVector VertexPosition = EditableMesh->SubdivisionLimitData.VertexPositions[ QuadVertex.VertexPositionIndex ];
+
+						FStaticMeshBuildVertex& StaticMeshVertex = StaticMeshBuildVertices[ NextVertexIndex ];
+						StaticMeshVertex.Position = VertexPosition;
+						StaticMeshVertex.TangentX = QuadVertex.VertexTangent;
+						StaticMeshVertex.TangentY = QuadVertex.VertexBinormalSign * FVector::CrossProduct( QuadVertex.VertexNormal, QuadVertex.VertexTangent );
+						StaticMeshVertex.TangentZ = QuadVertex.VertexNormal;
+
+						for( int32 UVIndex = 0; UVIndex < SubdivMeshTextureCoordinateCount; ++UVIndex )
+						{
+							StaticMeshVertex.UVs[ UVIndex ] = *( &QuadVertex.TextureCoordinate0 + UVIndex );
+						}
+
+						StaticMeshVertex.Color = QuadVertex.VertexColor;
+						if( StaticMeshVertex.Color != FColor::White )
+						{
+							bHasColor = true;
+						}
+
+						IndexBuffer.Add( NextVertexIndex++ );
+					}
+				}
+			}
+
+			++PolygonGroupIt;
+		}
+	}
+	else
+	{
+		// set up vertex buffer elements
+		StaticMeshBuildVertices.SetNum( EditableMesh->VertexInstances.GetMaxIndex() );
+
+		for( TSparseArray<FMeshVertexInstance>::TConstIterator It( EditableMesh->VertexInstances ); It; ++It )
+		{
+			if( It->Color != FColor::White )
+			{
+				bHasColor = true;
+			}
+
+			FStaticMeshBuildVertex& StaticMeshVertex = StaticMeshBuildVertices[ It.GetIndex() ];
+			const FMeshVertexInstance& VertexInstance = *It;
+
+			StaticMeshVertex.Position = EditableMesh->Vertices[ VertexInstance.VertexID.GetValue() ].VertexPosition;
+			StaticMeshVertex.TangentX = VertexInstance.Tangent;
+			StaticMeshVertex.TangentY = FVector::CrossProduct( VertexInstance.Normal, VertexInstance.Tangent ).GetSafeNormal() * VertexInstance.BinormalSign;
+			StaticMeshVertex.TangentZ = VertexInstance.Normal;
+			StaticMeshVertex.Color = VertexInstance.Color.ToFColor( true );
+			for( int32 UVIndex = 0; UVIndex < VertexInstance.VertexUVs.Num(); ++UVIndex )
+			{
+				StaticMeshVertex.UVs[ UVIndex ] = VertexInstance.VertexUVs[ UVIndex ];
+			}
+		}
+
+		// Set up index buffer
+		for( TSparseArray< FMeshPolygonGroup >::TConstIterator It( EditableMesh->PolygonGroups ); It; ++It )
+		{
+			const FMeshPolygonGroup& PolygonGroup = EditableMesh->PolygonGroups[ It.GetIndex() ];
+			FRenderingPolygonGroup& RenderingPolygonGroup = RenderingPolygonGroups[ It.GetIndex() ];
+
+			RenderingPolygonGroup.RenderingSectionIndex = StaticMeshLOD.Sections.Num();
+
+			// Create new rendering section
+			StaticMeshLOD.Sections.Add( FStaticMeshSection() );
+			FStaticMeshSection& StaticMeshSection = StaticMeshLOD.Sections.Last();
+
+			StaticMeshSection.FirstIndex = IndexBuffer.Num();
+			StaticMeshSection.NumTriangles = RenderingPolygonGroup.Triangles.GetMaxIndex();
+			check( RenderingPolygonGroup.Triangles.GetMaxIndex() <= RenderingPolygonGroup.MaxTriangles );
+
+			const int32 MaterialIndex = GetStaticMeshMaterialIndex( StaticMesh, PolygonGroup.Material );
+			check( MaterialIndex != INDEX_NONE );
+			StaticMeshSection.MaterialIndex = MaterialIndex;
+			StaticMeshSection.bEnableCollision = PolygonGroup.bEnableCollision;
+			StaticMeshSection.bCastShadow = PolygonGroup.bCastShadow;
+
+
+			if( RenderingPolygonGroup.Triangles.Num() > 0 )
+			{
+				IndexBuffer.Reserve( IndexBuffer.Num() + RenderingPolygonGroup.Triangles.GetMaxIndex() * 3 );
+				uint32 MinIndex = TNumericLimits< uint32 >::Max();
+				uint32 MaxIndex = TNumericLimits< uint32 >::Min();
+
+				// Find the first valid vertex instance index, so that we have a value we can use for our degenerates
+				check( RenderingPolygonGroup.Triangles.Num() > 0 );
+				const FVertexInstanceID FirstValidRenderingID = RenderingPolygonGroup.Triangles.CreateConstIterator()->VertexInstanceIDs[ 0 ];
+
+				for( int32 TriangleIndex = 0; TriangleIndex < RenderingPolygonGroup.Triangles.GetMaxIndex(); ++TriangleIndex )
+				{
+					if( RenderingPolygonGroup.Triangles.IsAllocated( TriangleIndex ) )
+					{
+						const FMeshTriangle& Triangle = RenderingPolygonGroup.Triangles[ TriangleIndex ];
+						for( int32 TriVert = 0; TriVert < 3; ++TriVert )
+						{
+							const uint32 RenderingVertexIndex = Triangle.VertexInstanceIDs[ TriVert ].GetValue();
+							IndexBuffer.Add( RenderingVertexIndex );
+							MinIndex = FMath::Min( MinIndex, RenderingVertexIndex );
+							MaxIndex = FMath::Max( MaxIndex, RenderingVertexIndex );
+						}
+					}
+					else
+					{
+						IndexBuffer.Add( FirstValidRenderingID.GetValue() );
+						IndexBuffer.Add( FirstValidRenderingID.GetValue() );
+						IndexBuffer.Add( FirstValidRenderingID.GetValue() );
+					}
+				}
+
+				StaticMeshSection.MinVertexIndex = MinIndex;
+				StaticMeshSection.MaxVertexIndex = MaxIndex;
+
+				// Add any index buffer padding.
+				// This can be necessary if we have just loaded an editable mesh which had a MaxTriangles count in the editable mesh section
+				// greater than the sparse array max size (i.e. an extra gap had been reserved for tris).
+				const int32 IndexBufferPadding = RenderingPolygonGroup.MaxTriangles - RenderingPolygonGroup.Triangles.GetMaxIndex();
+				if( IndexBufferPadding > 0 )
+				{
+					IndexBuffer.AddZeroed( IndexBufferPadding * 3 );
+				}
+			}
+			else
+			{
+				// No triangles in this section
+				StaticMeshSection.MinVertexIndex = 0;
+				StaticMeshSection.MaxVertexIndex = 0;
+			}
+		}
+	}
+
+	// Figure out which index buffer stride we need
+	bool bNeeds32BitIndices = false;
+	for( const FStaticMeshSection& StaticMeshSection : StaticMeshLOD.Sections )
+	{
+		if( StaticMeshSection.MaxVertexIndex > TNumericLimits<uint16>::Max() )
+		{
+			bNeeds32BitIndices = true;
+		}
+	}
+	const EIndexBufferStride::Type IndexBufferStride = bNeeds32BitIndices ? EIndexBufferStride::Force32Bit : EIndexBufferStride::Force16Bit;
+
+	StaticMeshLOD.PositionVertexBuffer.Init( StaticMeshBuildVertices );
+	StaticMeshLOD.VertexBuffer.Init( StaticMeshBuildVertices, EditableMesh->GetTextureCoordinateCount() );
+
+	if( bHasColor )
+	{
+		StaticMeshLOD.ColorVertexBuffer.Init( StaticMeshBuildVertices );
+	}
+	else
+	{
+		StaticMeshLOD.ColorVertexBuffer.InitFromSingleColor( FColor::White, StaticMeshBuildVertices.Num() );
+	}
+
+	StaticMeshLOD.IndexBuffer.SetIndices( IndexBuffer, IndexBufferStride );
+
+	// @todo mesheditor: support the other index buffer types
+	StaticMeshLOD.ReversedIndexBuffer.SetIndices( TArray< uint32 >(), IndexBufferStride );
+	StaticMeshLOD.DepthOnlyIndexBuffer.SetIndices( TArray< uint32 >(), IndexBufferStride );
+	StaticMeshLOD.ReversedDepthOnlyIndexBuffer.SetIndices( TArray< uint32 >(), IndexBufferStride );
+	StaticMeshLOD.WireframeIndexBuffer.SetIndices( TArray< uint32 >(), IndexBufferStride );
+	StaticMeshLOD.AdjacencyIndexBuffer.SetIndices( TArray< uint32 >(), IndexBufferStride );
+
+	StaticMeshLOD.bHasAdjacencyInfo = false;
+	StaticMeshLOD.bHasDepthOnlyIndices = false;
+	StaticMeshLOD.bHasReversedIndices = false;
+	StaticMeshLOD.bHasReversedDepthOnlyIndices = false;
+	StaticMeshLOD.DepthOnlyNumTriangles = 0;
+}
+
+
+
+void UEditableStaticMeshAdapter::Serialize( FArchive& Ar )
+{
+	Super::Serialize( Ar );
+
+	Ar.UsingCustomVersion( FEditableMeshCustomVersion::GUID );
+
+	SerializeSparseArray( Ar, RenderingPolygons );
+	SerializeSparseArray( Ar, RenderingPolygonGroups );
+}
+
+
+void UEditableStaticMeshAdapter::OnStartModification( const UEditableMesh* EditableMesh, const EMeshModificationType MeshModificationType, const EMeshTopologyChange MeshTopologyChange )
+{
+	// @todo mesheditor undo: We're not using traditional transactions to undo mesh changes yet, but we still want to dirty the mesh package
+	// Also, should we even need the Initializing type? Should we not wait for the first modification before dirtying the package?
+	if( 0 )
+	{
+		this->SetFlags( RF_Transactional );
+		this->Modify();
+
+		StaticMesh->SetFlags( RF_Transactional );
+		StaticMesh->Modify();
+	}
+	else
+	{
+		StaticMesh->MarkPackageDirty();
+	}
+}
+
+
+void UEditableStaticMeshAdapter::OnRebuildRenderMeshStart( const UEditableMesh* EditableMesh, const bool bRefreshBounds, const bool bInvalidateLighting )
+{
+	// We're changing the mesh itself, so ALL static mesh components in the scene will need
+	// to be unregistered for this (and reregistered afterwards.)
+	RecreateRenderStateContext = MakeShareable( new FStaticMeshComponentRecreateRenderStateContext( StaticMesh, bInvalidateLighting, bRefreshBounds ) );
+
+	// Release the static mesh's resources.
+	StaticMesh->ReleaseResources();
+
+	// Flush the resource release commands to the rendering thread to ensure that the build doesn't occur while a resource is still
+	// allocated, and potentially accessing the UStaticMesh.
+	StaticMesh->ReleaseResourcesFence.Wait();
+}
+
+
+void UEditableStaticMeshAdapter::OnEndModification( const UEditableMesh* EditableMesh )
+{
+	// nothing to do here
+}
+
+
+void UEditableStaticMeshAdapter::OnRebuildRenderMeshFinish( const UEditableMesh* EditableMesh, const bool bUpdateCollision )
+{
+	UpdateBoundsAndCollision( EditableMesh, bUpdateCollision );
+
+	StaticMesh->InitResources();
+
+	// NOTE: This can call InvalidateLightingCache() on all components using this mesh, causing Modify() to be 
+	// called on those components!  Just something to be aware of when EndModification() is called within
+	// an undo transaction.
+	RecreateRenderStateContext.Reset();
+}
+
+
+void UEditableStaticMeshAdapter::OnReindexElements( const UEditableMesh* EditableMesh, const FElementIDRemappings& Remappings )
+{
+	RemapSparseArrayElements( RenderingPolygons, Remappings.NewPolygonIndexLookup );
+	RemapSparseArrayElements( RenderingPolygonGroups, Remappings.NewPolygonGroupIndexLookup );
+
+	// Always compact the rendering triangles
+	for( TSparseArray<FRenderingPolygonGroup>::TIterator It( RenderingPolygonGroups ); It; ++It )
+	{
+		FRenderingPolygonGroup& RenderingPolygonGroup = *It;
+		const FMeshPolygonGroup& PolygonGroup = EditableMesh->PolygonGroups[ It.GetIndex() ];
+
+		TSparseArray<FTriangleID> TriangleRemappings;
+		CompactSparseArrayElements( RenderingPolygonGroup.Triangles, TriangleRemappings );
+
+		for( FMeshTriangle& Triangle : RenderingPolygonGroup.Triangles )
+		{
+			for( FVertexInstanceID& VertexInstanceID : Triangle.VertexInstanceIDs )
+			{
+				VertexInstanceID = Remappings.NewVertexInstanceIndexLookup[ VertexInstanceID.GetValue() ];
+			}
+		}
+
+		RenderingPolygonGroup.MaxTriangles = RenderingPolygonGroup.Triangles.GetMaxIndex();
+
+		// Fix up references in referencing polygons
+		for( const FPolygonID PolygonID : PolygonGroup.Polygons )
+		{
+			FRenderingPolygon& RenderingPolygon = RenderingPolygons[ PolygonID.GetValue() ];
+			for( FTriangleID& TriangleID : RenderingPolygon.TriangulatedPolygonTriangleIndices )
+			{
+				TriangleID = TriangleRemappings[ TriangleID.GetValue() ];
+			}
+		}
+	}
+}
+
+
+bool UEditableStaticMeshAdapter::IsCommitted( const UEditableMesh* EditableMesh ) const
+{
+	return StaticMesh->EditableMesh == EditableMesh;
+}
+
+
+bool UEditableStaticMeshAdapter::IsCommittedAsInstance( const UEditableMesh* EditableMesh ) const
+{
+	return StaticMesh != OriginalStaticMesh;
+}
+
+
+void UEditableStaticMeshAdapter::OnCommit( UEditableMesh* EditableMesh )
+{
+	if( !IsCommitted( EditableMesh ) )
+	{
+		// Move the editable mesh to an inner of the static mesh, and set the static mesh's EditableMesh property.
+		EditableMesh->Rename( nullptr, StaticMesh, REN_DontCreateRedirectors );
+		StaticMesh->EditableMesh = EditableMesh;
+	}
+}
+
+
+UEditableMesh* UEditableStaticMeshAdapter::OnCommitInstance( UEditableMesh* EditableMesh, UPrimitiveComponent* ComponentToInstanceTo )
+{
+	UStaticMeshComponent* StaticMeshComponent = Cast<UStaticMeshComponent>( ComponentToInstanceTo );
+
+	if( StaticMeshComponent )
+	{
+		// Duplicate the static mesh, putting it as an *inner* of the static mesh component.
+		// This is no longer a persistent asset, so clear the appropriate flags.
+		UStaticMesh* NewStaticMesh = DuplicateObject( OriginalStaticMesh, StaticMeshComponent );
+		NewStaticMesh->ClearFlags( RF_Public | RF_Standalone );
+
+		// Point the static mesh component to the new static mesh instance we just made for it
+		StaticMeshComponent->SetStaticMesh( NewStaticMesh );
+
+		// Duplicate this editable mesh to a new instance inside the new static mesh instance, and set the static mesh's EditableMesh property.
+		UEditableMesh* NewEditableMesh = DuplicateObject( EditableMesh, NewStaticMesh );
+
+		// Look for the corresponding adapter instance in the duplicated mesh.
+		const int32 AdapterIndex = EditableMesh->Adapters.Find( this );
+		check( AdapterIndex != INDEX_NONE );
+		UEditableStaticMeshAdapter* NewAdapter = Cast<UEditableStaticMeshAdapter>( NewEditableMesh->Adapters[ AdapterIndex ] );
+
+		NewStaticMesh->EditableMesh = NewEditableMesh;
+		NewAdapter->StaticMesh = NewStaticMesh;
+
+		// Update the submesh address which will have changed now it's been instanced
+		NewEditableMesh->SetSubMeshAddress( UEditableMeshFactory::MakeSubmeshAddress( StaticMeshComponent, EditableMesh->SubMeshAddress.LODIndex ) );
+		NewEditableMesh->RebuildRenderMesh();
+
+		return NewEditableMesh;
+	}
+
+	return nullptr;
+}
+
+
+void UEditableStaticMeshAdapter::OnRevert( UEditableMesh* EditableMesh )
+{
+	// @todo
+}
+
+
+UEditableMesh* UEditableStaticMeshAdapter::OnRevertInstance( UEditableMesh* EditableMesh )
+{
+	// @todo
+	return nullptr;
+}
+
+
+void UEditableStaticMeshAdapter::OnPropagateInstanceChanges( UEditableMesh* EditableMesh )
+{
+	check( false );		// @todo mesheditor: fix me
+	//if( IsCommittedAsInstance() )
+	//{
+	//	// @todo mesheditor: we can only generate submesh addresses from a component. Since we don't have a component, we create a dummy one.
+	//	// Not really fond of this.
+	//	// Explore other possibilities, e.g. constructing a submesh address by hand (although the contents of MeshObjectPtr are supposed to be opaque)
+	//	UStaticMeshComponent* DummyComponent = NewObject<UStaticMeshComponent>();
+	//	DummyComponent->SetStaticMesh( OriginalStaticMesh );
+
+	//	UEditableStaticMeshAdapter* NewEditableMesh = DuplicateObject( this, OriginalStaticMesh );
+	//	OriginalStaticMesh->EditableMesh = NewEditableMesh;
+	//	NewEditableMesh->StaticMesh = OriginalStaticMesh;
+	//	NewEditableMesh->SetSubMeshAddress( UEditableMeshFactory::MakeSubmeshAddress( DummyComponent, SubMeshAddress.LODIndex ) );
+	//	NewEditableMesh->RebuildRenderMesh();
+	//}
+}
+
+
+void UEditableStaticMeshAdapter::UpdateBoundsAndCollision( const UEditableMesh* EditableMesh, const bool bUpdateCollision )
+{
+	// @todo mesheditor: we will need to create a new DDC key once we are able to edit placed instances individually.
+	// Will need to find a way of deriving the key based on the mesh key and an instance number which remains constant,
+	// otherwise we risk filling the DDC with junk (i.e. using vertex positions etc is not scalable).
+
+	// Compute a new bounding box
+	// @todo mesheditor perf: Only do this if the bounds may have changed (need hinting)
+	{
+		FStaticMeshRenderData& StaticMeshRenderData = *StaticMesh->RenderData;
+
+		FBoxSphereBounds BoundingBoxAndSphere;
+
+		// @todo mesheditor LODs: Really we should store the bounds of LOD0 inside the static mesh.  Our editable mesh might be for a different LOD.
+
+		// If we're in subdivision preview mode, use the bounds of the base cage mesh, so that simple collision
+		// queries will always include the base cage, even though the actual mesh geometry might be quite a bit smaller.
+		// This also relies on us specifically querying against the simple collision, which we do in a second pass after
+		// looking for meshes using a complex collision trace.
+		// @todo mesheditor: Ideally we are not storing an inflated bounds here just for base cage editor interaction
+		if( EditableMesh->IsPreviewingSubdivisions() )
+		{
+			BoundingBoxAndSphere = EditableMesh->ComputeBoundingBoxAndSphere();
+		}
+		else
+		{
+			FBox BoundingBox;
+			BoundingBox.Init();
+
+			// Could improve performance here if necessary:
+			// 1) cache polygon IDs per vertex (in order to quickly reject orphans) and just iterate vertex array; or
+			// 2) cache bounding box per polygon
+			// There are other cases where having polygon adjacency information (1) might be useful, so it's maybe worth considering.
+
+			for( const FMeshPolygon& Polygon : EditableMesh->Polygons )
+			{
+				for( const FVertexInstanceID VertexInstanceID : Polygon.PerimeterContour.VertexInstanceIDs )
+				{
+					const FMeshVertexInstance& VertexInstance = EditableMesh->VertexInstances[ VertexInstanceID.GetValue() ];
+					BoundingBox += EditableMesh->Vertices[ VertexInstance.VertexID.GetValue() ].VertexPosition;
+				}
+			}
+
+			BoundingBox.GetCenterAndExtents( /* Out */ BoundingBoxAndSphere.Origin, /* Out */ BoundingBoxAndSphere.BoxExtent );
+
+			// Calculate the bounding sphere, using the center of the bounding box as the origin.
+			BoundingBoxAndSphere.SphereRadius = 0.0f;
+
+			for( const FMeshPolygon& Polygon : EditableMesh->Polygons )
+			{
+				for( const FVertexInstanceID VertexInstanceID : Polygon.PerimeterContour.VertexInstanceIDs )
+				{
+					const FMeshVertexInstance& VertexInstance = EditableMesh->VertexInstances[ VertexInstanceID.GetValue() ];
+					const FVector VertexPosition = EditableMesh->Vertices[ VertexInstance.VertexID.GetValue() ].VertexPosition;
+
+					BoundingBoxAndSphere.SphereRadius = FMath::Max( ( VertexPosition - BoundingBoxAndSphere.Origin ).Size(), BoundingBoxAndSphere.SphereRadius );
+				}
+			}
+		}
+	
+		StaticMeshRenderData.Bounds = BoundingBoxAndSphere;
+		StaticMesh->CalculateExtendedBounds();
+	}
+
+
+	// Refresh collision (only if the interaction has finished though -- this is really expensive!)
+	if( bUpdateCollision )
+	{
+
+		// @todo mesheditor collision: We're wiping the existing simplified collision and generating a simple bounding
+		// box collision, since that's the best we can do without impacting performance.  We always using visibility (complex)
+		// collision for traces while mesh editing (for hover/selection), so simplified collision isn't really important.
+		const bool bRecreateSimplifiedCollision = true;
+
+		if( StaticMesh->BodySetup == nullptr )
+		{
+			StaticMesh->CreateBodySetup();
+		}
+
+		UBodySetup* BodySetup = StaticMesh->BodySetup;
+
+		// NOTE: We don't bother calling Modify() on the BodySetup as EndModification() will rebuild this guy after every undo
+		// BodySetup->Modify();
+
+		if( bRecreateSimplifiedCollision )
+		{
+			if( BodySetup->AggGeom.GetElementCount() > 0 )
+			{
+				BodySetup->RemoveSimpleCollision();
+			}
+		}
+
+		BodySetup->InvalidatePhysicsData();
+
+		if( bRecreateSimplifiedCollision )
+		{
+			const FBoxSphereBounds Bounds = StaticMesh->GetBounds();
+
+			FKBoxElem BoxElem;
+			BoxElem.Center = Bounds.Origin;
+			BoxElem.X = Bounds.BoxExtent.X * 2.0f;
+			BoxElem.Y = Bounds.BoxExtent.Y * 2.0f;
+			BoxElem.Z = Bounds.BoxExtent.Z * 2.0f;
+			BodySetup->AggGeom.BoxElems.Add( BoxElem );
+		}
+
+		// Update all static mesh components that are using this mesh
+		// @todo mesheditor perf: This is a pretty heavy operation, and overlaps with what we're already doing in RecreateRenderStateContext
+		// a little bit.  Ideally we do everything in a single pass.  Furthermore, if this could be updated lazily it would be faster.
+		{
+			for( FObjectIterator Iter( UStaticMeshComponent::StaticClass() ); Iter; ++Iter )
+			{
+				UStaticMeshComponent* StaticMeshComponent = Cast<UStaticMeshComponent>( *Iter );
+				if( StaticMeshComponent->GetStaticMesh() == StaticMesh )
+				{
+					// it needs to recreate IF it already has been created
+					if( StaticMeshComponent->IsPhysicsStateCreated() )
+					{
+						StaticMeshComponent->RecreatePhysicsState();
+					}
+				}
+			}
+		}
+	}
+}
+
+
+void UEditableStaticMeshAdapter::OnSetVertexAttribute( const UEditableMesh* EditableMesh, const FVertexID VertexID, const FName AttributeName, const int32 AttributeIndex, const FVector4 AttributeValue )
+{
+	const FMeshVertex& Vertex = EditableMesh->Vertices[ VertexID.GetValue() ];
+	FStaticMeshLODResources& StaticMeshLOD = GetStaticMeshLOD();
+
+	if( AttributeName == UEditableMeshAttribute::VertexPosition() )
+	{
+		// @todo mesheditor: eventually break out subdivided mesh into a different adapter which handles things differently?
+		// (may also want different component eventually)
+		if( !EditableMesh->IsPreviewingSubdivisions() )
+		{
+			// Set the vertex buffer position of all of the vertex instances for this editable vertex
+			for( const FVertexInstanceID VertexInstanceID : Vertex.VertexInstanceIDs )
+			{
+				check( EditableMesh->VertexInstances.IsAllocated( VertexInstanceID.GetValue() ) );
+				StaticMeshLOD.PositionVertexBuffer.VertexPosition( VertexInstanceID.GetValue() ) = FVector( AttributeValue );
+			}
+		}
+	}
+}
+
+
+void UEditableStaticMeshAdapter::OnSetEdgeAttribute( const UEditableMesh* EditableMesh, const FEdgeID EdgeID, const FName AttributeName, const int32 AttributeIndex, const FVector4 AttributeValue )
+{
+	// Nothing to do here
+}
+
+
+//int32 UEditableStaticMeshAdapter::GetTriangleCount( const FPolygonGroupID PolygonGroupID ) const
+//{
+//	checkSlow( RenderingPolygonGroups.IsAllocated( PolygonGroupID.GetValue() ) );
+//	const FRenderingPolygonGroup& RenderingPolygonGroup = RenderingPolygonGroups[ PolygonGroupID.GetValue() ];
+//
+//	return RenderingPolygonGroup.Triangles.Num();
+//}
+//
+//
+//int32 UEditableStaticMeshAdapter::GetTriangleArraySize( const FPolygonGroupID PolygonGroupID ) const
+//{
+//	checkSlow( RenderingPolygonGroups.IsAllocated( PolygonGroupID.GetValue() ) );
+//	const FRenderingPolygonGroup& RenderingPolygonGroup = RenderingPolygonGroups[ PolygonGroupID.GetValue() ];
+//
+//	return RenderingPolygonGroup.Triangles.GetMaxIndex();
+//}
+
+
+void UEditableStaticMeshAdapter::OnSetVertexInstanceAttribute( const UEditableMesh* EditableMesh, const FVertexInstanceID VertexInstanceID, const FName AttributeName, const int32 AttributeIndex, const FVector4 AttributeValue )
+{
+	const FMeshVertexInstance& VertexInstance = EditableMesh->VertexInstances[ VertexInstanceID.GetValue() ];
+	FStaticMeshLODResources& StaticMeshLOD = GetStaticMeshLOD();
+
+	if( AttributeName == UEditableMeshAttribute::VertexNormal() ||
+		AttributeName == UEditableMeshAttribute::VertexTangent() ||
+		AttributeName == UEditableMeshAttribute::VertexBinormalSign() )
+	{
+		if( !EditableMesh->IsPreviewingSubdivisions() )
+		{
+			// @todo mesheditor perf: SetVertexTangents() and VertexTangentX/Y() functions actually does a bit of work to compute the basis every time. 
+			// Ideally we can get/set this stuff directly to improve performance.  This became slower after high precision basis values were added.
+			// @todo mesheditor perf: this is even more pertinent now we already have the binormal sign!
+			StaticMeshLOD.VertexBuffer.SetVertexTangents(
+				VertexInstanceID.GetValue(),
+				VertexInstance.Tangent,
+				FVector::CrossProduct( VertexInstance.Normal, VertexInstance.Tangent ).GetSafeNormal() * VertexInstance.BinormalSign,
+				VertexInstance.Normal );
+		}
+	}
+	else if( AttributeName == UEditableMeshAttribute::VertexTextureCoordinate() )
+	{
+		if( !EditableMesh->IsPreviewingSubdivisions() )
+		{
+			check( AttributeIndex < EditableMesh->GetTextureCoordinateCount() );
+			StaticMeshLOD.VertexBuffer.SetVertexUV( VertexInstanceID.GetValue(), AttributeIndex, VertexInstance.VertexUVs[ AttributeIndex ] );
+		}
+	}
+	else if( AttributeName == UEditableMeshAttribute::VertexColor() )
+	{
+		if( !EditableMesh->IsPreviewingSubdivisions() )
+		{
+			const FColor NewColor = VertexInstance.Color.ToFColor( true );
+
+			if( StaticMeshLOD.ColorVertexBuffer.GetNumVertices() != EditableMesh->VertexInstances.GetMaxIndex() )
+			{
+				if( VertexInstance.Color != FLinearColor::White )
+				{
+					// Until now, we haven't needed a vertex color buffer.
+					// Force one to be generated now that we have a non-white vertex in the mesh.
+					OnRebuildRenderMesh( EditableMesh );
+				}
+			}
+			else
+			{
+				StaticMeshLOD.ColorVertexBuffer.VertexColor( VertexInstanceID.GetValue() ) = NewColor;
+			}
+		}
+	}
+}
+
+
+void UEditableStaticMeshAdapter::OnCreateEmptyVertexRange( const UEditableMesh* EditableMesh, const TArray<FVertexID>& VertexIDs )
+{
+}
+
+
+void UEditableStaticMeshAdapter::OnCreateVertices( const UEditableMesh* EditableMesh, const TArray<FVertexID>& VertexIDs )
+{
+}
+
+
+void UEditableStaticMeshAdapter::OnCreateVertexInstances( const UEditableMesh* EditableMesh, const TArray<FVertexInstanceID>& VertexInstanceIDs )
+{
+	if( !EditableMesh->IsPreviewingSubdivisions() )
+	{
+		FStaticMeshLODResources& StaticMeshLOD = GetStaticMeshLOD();
+		const int32 NumUVs = EditableMesh->GetTextureCoordinateCount();
+		const bool bHasColors = StaticMeshLOD.ColorVertexBuffer.GetNumVertices() > 0;
+
+		// Determine if we need to grow the render buffers
+		const int32 OldVertexBufferRenderingVertexCount = StaticMeshLOD.PositionVertexBuffer.GetNumVertices();
+		const int32 NumNewVertexBufferRenderingVertices = FMath::Max( 0, EditableMesh->VertexInstances.GetMaxIndex() - OldVertexBufferRenderingVertexCount );
+
+		static TArray<FStaticMeshBuildVertex> RenderingVerticesToAppend;
+		RenderingVerticesToAppend.SetNumUninitialized( NumNewVertexBufferRenderingVertices, false );
+
+		for( const FVertexInstanceID VertexInstanceID : VertexInstanceIDs )
+		{
+			const FMeshVertexInstance& VertexInstance = EditableMesh->VertexInstances[ VertexInstanceID.GetValue() ];
+			const FMeshVertex& ReferencedVertex = EditableMesh->Vertices[ VertexInstance.VertexID.GetValue() ];
+
+			const int32 NewRenderingVertexIndex = VertexInstanceID.GetValue();	// Rendering vertex indices are the same as vertex instance IDs
+
+			if( NewRenderingVertexIndex < OldVertexBufferRenderingVertexCount )
+			{
+				// Rendering vertex is within the already allocated buffer. Initialize the new vertices to some defaults
+				StaticMeshLOD.PositionVertexBuffer.VertexPosition( NewRenderingVertexIndex ) = ReferencedVertex.VertexPosition;
+				StaticMeshLOD.VertexBuffer.SetVertexTangents(
+					NewRenderingVertexIndex,
+					FVector::ZeroVector,
+					FVector::ZeroVector,
+					FVector::ZeroVector );
+
+				for( int32 UVIndex = 0; UVIndex < NumUVs; ++UVIndex )
+				{
+					StaticMeshLOD.VertexBuffer.SetVertexUV( NewRenderingVertexIndex, UVIndex, FVector2D::ZeroVector );
+				}
+
+				if( bHasColors )
+				{
+					StaticMeshLOD.ColorVertexBuffer.VertexColor( NewRenderingVertexIndex ) = FColor::White;
+				}
+			}
+			else
+			{
+				// Rendering vertex needs to be added in a new block
+				const int32 AppendVertexNumber = NewRenderingVertexIndex - OldVertexBufferRenderingVertexCount;
+				check( AppendVertexNumber >= 0 && AppendVertexNumber < NumNewVertexBufferRenderingVertices );
+				FStaticMeshBuildVertex& RenderingVertexToAppend = RenderingVerticesToAppend[ AppendVertexNumber ];
+
+				// Initialize the new vertices to some defaults
+				RenderingVertexToAppend.Position = ReferencedVertex.VertexPosition;
+				RenderingVertexToAppend.TangentX = FVector::ZeroVector;
+				RenderingVertexToAppend.TangentY = FVector::ZeroVector;
+				RenderingVertexToAppend.TangentZ = FVector::ZeroVector;
+
+				for( int32 UVIndex = 0; UVIndex < NumUVs; ++UVIndex )
+				{
+					RenderingVertexToAppend.UVs[ UVIndex ] = FVector2D::ZeroVector;
+				}
+
+				RenderingVertexToAppend.Color = FColor::White;
+			}
+		}
+
+		if( RenderingVerticesToAppend.Num() > 0 )
+		{
+			StaticMeshLOD.VertexBuffer.AppendVertices( RenderingVerticesToAppend.GetData(), RenderingVerticesToAppend.Num() );
+			StaticMeshLOD.PositionVertexBuffer.AppendVertices( RenderingVerticesToAppend.GetData(), RenderingVerticesToAppend.Num() );
+
+			if( bHasColors )
+			{
+				StaticMeshLOD.ColorVertexBuffer.AppendVertices( RenderingVerticesToAppend.GetData(), RenderingVerticesToAppend.Num() );
+			}
+		}
+	}
+}
+
+
+void UEditableStaticMeshAdapter::OnCreateEdges( const UEditableMesh* EditableMesh, const TArray<FEdgeID>& EdgeIDs )
+{
+	// Nothing to do here for now
+}
+
+
+void UEditableStaticMeshAdapter::OnCreatePolygons( const UEditableMesh* EditableMesh, const TArray<FPolygonID>& PolygonIDs )
+{
+	// Add mirror polygons for static mesh adapter
+	for( const FPolygonID PolygonID : PolygonIDs )
+	{
+		RenderingPolygons.Insert( PolygonID.GetValue(), FRenderingPolygon() );
+	}
+}
+
+
+void UEditableStaticMeshAdapter::OnChangePolygonVertexInstances( const UEditableMesh* EditableMesh, const TArray<FPolygonID>& PolygonIDs )
+{
+	OnRetriangulatePolygons( EditableMesh, PolygonIDs );
+}
+
+
+void UEditableStaticMeshAdapter::OnRetriangulatePolygons( const UEditableMesh* EditableMesh, const TArray<FPolygonID>& PolygonIDs )
+{
+	for( const FPolygonID PolygonID : PolygonIDs )
+	{
+		const FPolygonGroupID PolygonGroupID = EditableMesh->GetGroupForPolygon( PolygonID );
+
+		const FMeshPolygonGroup& PolygonGroup = EditableMesh->PolygonGroups[ PolygonGroupID.GetValue() ];
+		const FMeshPolygon& Polygon = EditableMesh->Polygons[ PolygonID.GetValue() ];
+
+		FRenderingPolygonGroup& RenderingPolygonGroup = RenderingPolygonGroups[ PolygonGroupID.GetValue() ];
+		FRenderingPolygon& RenderingPolygon = RenderingPolygons[ PolygonID.GetValue() ];
+
+		const TArray<FMeshTriangle>& Triangles = Polygon.Triangles;
+
+		// Check to see whether the index buffer needs to be updated
+		bool bNeedsUpdatedTriangles = false;
+		{
+			if( RenderingPolygon.TriangulatedPolygonTriangleIndices.Num() != Triangles.Num() )
+			{
+				// Triangle count has changed, so we definitely need new triangles!
+				bNeedsUpdatedTriangles = true;
+			}
+			else
+			{
+				// See if the triangulation has changed even if the number of triangles is the same
+				for( int32 Index = 0; Index < Triangles.Num(); ++Index )
+				{
+					const FMeshTriangle& OldTriangle = RenderingPolygonGroup.Triangles[ RenderingPolygon.TriangulatedPolygonTriangleIndices[ Index ].GetValue() ];
+					const FMeshTriangle& NewTriangle = Triangles[ Index ];
+
+					if( OldTriangle.VertexInstanceIDs[ 0 ] != NewTriangle.VertexInstanceIDs[ 0 ] ||
+						OldTriangle.VertexInstanceIDs[ 1 ] != NewTriangle.VertexInstanceIDs[ 1 ] ||
+						OldTriangle.VertexInstanceIDs[ 2 ] != NewTriangle.VertexInstanceIDs[ 2 ] )
+					{
+						bNeedsUpdatedTriangles = true;
+						break;
+					}
+				}
+			}
+		}
+
+		// Has anything changed?
+		if( bNeedsUpdatedTriangles )
+		{
+			const uint32 RenderingSectionIndex = RenderingPolygonGroup.RenderingSectionIndex;
+			FStaticMeshLODResources& StaticMeshLOD = GetStaticMeshLOD();
+			FStaticMeshSection& RenderingSection = StaticMeshLOD.Sections[ RenderingSectionIndex ];
+
+			// Remove the old triangles
+			DeletePolygonTriangles( EditableMesh, PolygonID );
+
+			// Add new triangles
+			{
+				// This is the number of triangles we are about to add
+				const int32 NumNewTriangles = Triangles.Num();
+
+				// This is the number of entries currently unused in the Triangles sparse array
+				const int32 NumFreeTriangles = RenderingPolygonGroup.Triangles.GetMaxIndex() - RenderingPolygonGroup.Triangles.Num();
+
+				// This is the number of triangles we need to make extra space for (in both the sparse array and the index buffer)
+				const int32 NumTrianglesToReserve = FMath::Max( 0, NumNewTriangles - NumFreeTriangles );
+
+				// This is the number of triangles we will need to have allocated in the index buffer after adding the new triangles
+				const int32 NewTotalTriangles = RenderingPolygonGroup.Triangles.GetMaxIndex() + NumTrianglesToReserve;
+
+				// Reserve extra triangles if necessary.
+				if( NumTrianglesToReserve > 0 )
+				{
+					RenderingPolygonGroup.Triangles.Reserve( NewTotalTriangles );
+				}
+
+				// Keep track of new min/max vertex indices
+				int32 MinVertexIndex = RenderingSection.MinVertexIndex;
+				int32 MaxVertexIndex = RenderingSection.MaxVertexIndex;
+
+				// Create empty triangles for all of the new triangles we need, and keep track of their triangle indices
+				static TArray<int32> NewTriangleIndices;
+				{
+					NewTriangleIndices.SetNumUninitialized( NumNewTriangles, false );
+
+					for( int32 TriangleToAddNumber = 0; TriangleToAddNumber < NumNewTriangles; ++TriangleToAddNumber )
+					{
+						const int32 NewTriangleIndex = RenderingPolygonGroup.Triangles.Add( FMeshTriangle() );
+						NewTriangleIndices[ TriangleToAddNumber ] = NewTriangleIndex;
+
+						FMeshTriangle& NewTriangle = RenderingPolygonGroup.Triangles[ NewTriangleIndex ];
+						for( int32 TriangleVertexNumber = 0; TriangleVertexNumber < 3; ++TriangleVertexNumber )
+						{
+							const FVertexInstanceID VertexInstanceID = Triangles[ TriangleToAddNumber ].VertexInstanceIDs[ TriangleVertexNumber ];
+							NewTriangle.VertexInstanceIDs[ TriangleVertexNumber ] = VertexInstanceID;
+							MinVertexIndex = FMath::Min( MinVertexIndex, VertexInstanceID.GetValue() );
+							MaxVertexIndex = FMath::Max( MaxVertexIndex, VertexInstanceID.GetValue() );
+						}
+
+						RenderingPolygon.TriangulatedPolygonTriangleIndices.Add( FTriangleID( NewTriangleIndex ) );
+					}
+				}
+
+				// Update the index buffer format if the index range exceeds 16 bit values.
+				if( !EditableMesh->IsPreviewingSubdivisions() )
+				{
+					UpdateIndexBufferFormatIfNeeded( Triangles );
+				}
+
+				// If we need more space in the index buffer for this section, allocate it here
+				if( NewTotalTriangles > RenderingPolygonGroup.MaxTriangles )
+				{
+					const int32 NumExtraTriangles = NewTotalTriangles + IndexBufferInterSectionGap - RenderingPolygonGroup.MaxTriangles;
+
+					// Get current number of triangles allocated for this section
+					const int32 MaxTriangles = RenderingPolygonGroup.MaxTriangles;
+					RenderingPolygonGroup.MaxTriangles += NumExtraTriangles;
+
+					if( !EditableMesh->IsPreviewingSubdivisions() )
+					{
+						const uint32 FirstIndex = RenderingSection.FirstIndex;
+
+						// Make room in the index buffer for the extra triangles, and update the mesh section's maximum triangle count
+						StaticMeshLOD.IndexBuffer.InsertIndices( FirstIndex + MaxTriangles * 3, nullptr, NumExtraTriangles * 3 );
+
+						// Adjust first index for all subsequent render sections to account for the extra indices just inserted.
+						// It is guaranteed that index buffer indices are in the same order as the rendering sections.
+						const uint32 NumRenderingSections = StaticMeshLOD.Sections.Num();
+						uint32 NextRenderingSectionIndex = RenderingSectionIndex;
+						while( ++NextRenderingSectionIndex < NumRenderingSections )
+						{
+							check( StaticMeshLOD.Sections[ NextRenderingSectionIndex ].FirstIndex >= FirstIndex );
+							StaticMeshLOD.Sections[ NextRenderingSectionIndex ].FirstIndex += NumExtraTriangles * 3;
+						}
+					}
+				}
+
+				if( !EditableMesh->IsPreviewingSubdivisions() )
+				{
+					for( int32 TriangleToAddNumber = 0; TriangleToAddNumber < NumNewTriangles; ++TriangleToAddNumber )
+					{
+						const int32 NewTriangleIndex = NewTriangleIndices[ TriangleToAddNumber ];
+
+						for( int32 TriangleVertexNumber = 0; TriangleVertexNumber < 3; ++TriangleVertexNumber )
+						{
+							StaticMeshLOD.IndexBuffer.SetIndex(
+								FRenderingPolygonGroup::TriangleIndexToRenderingTriangleFirstIndex( RenderingSection, FTriangleID( NewTriangleIndex ) ) + TriangleVertexNumber,
+								RenderingPolygonGroup.Triangles[ NewTriangleIndex ].VertexInstanceIDs[ TriangleVertexNumber ].GetValue() );
+						}
+					}
+
+					if( NumTrianglesToReserve > 0 )
+					{
+						RenderingSection.NumTriangles += NumTrianglesToReserve;
+					}
+
+					RenderingSection.MinVertexIndex = MinVertexIndex;
+					RenderingSection.MaxVertexIndex = MaxVertexIndex;
+				}
+			}
+		}
+	}
+}
+
+
+void UEditableStaticMeshAdapter::OnDeleteVertexInstances( const UEditableMesh* EditableMesh, const TArray<FVertexInstanceID>& VertexInstanceIDs )
+{
+	// Nothing to do here
+}
+
+
+void UEditableStaticMeshAdapter::OnDeleteOrphanVertices( const UEditableMesh* EditableMesh, const TArray<FVertexID>& VertexIDs )
+{
+	// Nothing to do here
+}
+
+
+void UEditableStaticMeshAdapter::OnDeleteEdges( const UEditableMesh* EditableMesh, const TArray<FEdgeID>& EdgeIDs )
+{
+	// Nothing to do here
+}
+
+
+void UEditableStaticMeshAdapter::OnDeletePolygons( const UEditableMesh* EditableMesh, const TArray<FPolygonID>& PolygonIDs )
+{
+	for( const FPolygonID PolygonID : PolygonIDs )
+	{
+		// Removes all of a polygon's triangles (including rendering triangles from the index buffer.)
+		DeletePolygonTriangles( EditableMesh, PolygonID );
+
+		// Delete the polygon from the static mesh adapter mirror
+		RenderingPolygons.RemoveAt( PolygonID.GetValue() );
+	}
+}
+
+
+void UEditableStaticMeshAdapter::DeletePolygonTriangles( const UEditableMesh* EditableMesh, const FPolygonID PolygonID )
+{
+	const FPolygonGroupID PolygonGroupID = EditableMesh->GetGroupForPolygon( PolygonID );
+
+	FRenderingPolygonGroup& RenderingPolygonGroup = RenderingPolygonGroups[ PolygonGroupID.GetValue() ];
+	FRenderingPolygon& Polygon = RenderingPolygons[ PolygonID.GetValue() ];
+
+	const int32 NumTrianglesToRemove = Polygon.TriangulatedPolygonTriangleIndices.Num();
+	if( NumTrianglesToRemove > 0 )
+	{
+		bool bUpdateMinMax = false;
+
+		// Kill the polygon's rendering triangles in the static mesh
+
+		FStaticMeshLODResources& StaticMeshLOD = GetStaticMeshLOD();
+		const uint32 RenderingSectionIndex = RenderingPolygonGroup.RenderingSectionIndex;
+		FStaticMeshSection& RenderingSection = StaticMeshLOD.Sections[ RenderingSectionIndex ];
+
+		// Update the index buffer by removing entries, and the rendering sections with new section counts
+		if( !EditableMesh->IsPreviewingSubdivisions() )
+		{
+			// If the min or max vertex index is about to be deleted, refresh the range
+			const int32 MinVertexIndex = RenderingSection.MinVertexIndex;
+			const int32 MaxVertexIndex = RenderingSection.MaxVertexIndex;
+
+			for( const FTriangleID TriangleIndexToRemove : Polygon.TriangulatedPolygonTriangleIndices )
+			{
+				const FMeshTriangle& Triangle = RenderingPolygonGroup.Triangles[ TriangleIndexToRemove.GetValue() ];
+
+				for( const FVertexInstanceID VertexInstanceID : Triangle.VertexInstanceIDs )
+				{
+					if( VertexInstanceID.GetValue() == MinVertexIndex || VertexInstanceID.GetValue() == MaxVertexIndex )
+					{
+						bUpdateMinMax = true;
+						break;
+					}
+				}
+
+				if( bUpdateMinMax )
+				{
+					break;
+				}
+			}
+		}
+
+		// Remove all of the polygon's triangles from our editable mesh's triangle list.  While doing this, we'll keep
+		// track of all of the rendering mesh triangles that we'll need to remove later on.  We'll also figure out which
+		// vertex instances will need to be removed from their corresponding vertex
+		for( const FTriangleID TriangleIndexToRemove : Polygon.TriangulatedPolygonTriangleIndices )
+		{
+			// Remove this triangle from our editable mesh
+			RenderingPolygonGroup.Triangles.RemoveAt( TriangleIndexToRemove.GetValue() );
+		}
+
+		if( !EditableMesh->IsPreviewingSubdivisions() )
+		{
+			if( bUpdateMinMax )
+			{
+				int32 MinVertexIndex = TNumericLimits<int32>::Max();
+				int32 MaxVertexIndex = TNumericLimits<int32>::Min();
+
+				for( const FMeshTriangle& Triangle : RenderingPolygonGroup.Triangles )
+				{
+					for( const FVertexInstanceID VertexInstanceID : Triangle.VertexInstanceIDs )
+					{
+						if( VertexInstanceID.GetValue() < MinVertexIndex )
+						{
+							MinVertexIndex = VertexInstanceID.GetValue();
+						}
+
+						if( VertexInstanceID.GetValue() > MaxVertexIndex )
+						{
+							MaxVertexIndex = VertexInstanceID.GetValue();
+						}
+					}
+				}
+
+				RenderingSection.MinVertexIndex = MinVertexIndex;
+				RenderingSection.MaxVertexIndex = MaxVertexIndex;
+			}
+
+			// @todo mesheditor urgent: What about other index buffers in the mesh (DepthOnlyIndexBuffer, Wireframe, etc.)  We need to remove our triangles from those too!
+
+			for( const FTriangleID SectionTriangleIDToRemove : Polygon.TriangulatedPolygonTriangleIndices )
+			{
+				const uint32 RenderingTriangleFirstVertexIndex = FRenderingPolygonGroup::TriangleIndexToRenderingTriangleFirstIndex( RenderingSection, SectionTriangleIDToRemove );
+
+				// Make the indices degenerate.  We don't want to actually remove the indices from the index buffer, as that's can
+				// be a really slow operation.  The mesh can be compacted later on to free up the memory.
+				for( int32 TriangleVertexNumber = 0; TriangleVertexNumber < 3; ++TriangleVertexNumber )
+				{
+					StaticMeshLOD.IndexBuffer.SetIndex( RenderingTriangleFirstVertexIndex + TriangleVertexNumber, RenderingSection.MinVertexIndex );
+				}
+			}
+		}
+
+		Polygon.TriangulatedPolygonTriangleIndices.Reset();
+	}
+}
+
+
+inline const FStaticMeshLODResources& UEditableStaticMeshAdapter::GetStaticMeshLOD() const
+{
+	const FStaticMeshRenderData& StaticMeshRenderData = *StaticMesh->RenderData;
+	const FStaticMeshLODResources& StaticMeshLOD = StaticMeshRenderData.LODResources[ StaticMeshLODIndex ];
+	return StaticMeshLOD;
+}
+
+
+FStaticMeshLODResources& UEditableStaticMeshAdapter::GetStaticMeshLOD()
+{
+	FStaticMeshRenderData& StaticMeshRenderData = *StaticMesh->RenderData;
+	FStaticMeshLODResources& StaticMeshLOD = StaticMeshRenderData.LODResources[ StaticMeshLODIndex ];
+	return StaticMeshLOD;
+}
+
+
+void UEditableStaticMeshAdapter::OnCreatePolygonGroups( const UEditableMesh* EditableMesh, const TArray<FPolygonGroupID>& PolygonGroupIDs )
+{
+	for( const FPolygonGroupID PolygonGroupID : PolygonGroupIDs )
+	{
+		const FMeshPolygonGroup& PolygonGroup = EditableMesh->PolygonGroups[ PolygonGroupID.GetValue() ];
+
+		uint32 LODSectionIndex = 0;
+		if( !EditableMesh->IsPreviewingSubdivisions() )
+		{
+			// Need to create a new rendering section. This is added to the end of the array.
+			FStaticMeshLODResources& StaticMeshLOD = GetStaticMeshLOD();
+
+			LODSectionIndex = StaticMeshLOD.Sections.Emplace();
+			FStaticMeshSection& StaticMeshSection = StaticMeshLOD.Sections[ LODSectionIndex ];
+
+			// Initially the section is empty, and it occupies zero elements in the index buffer.
+			// It is still placed in the correct location within the index buffer, immediately following the previous section,
+			// as it is a requirement that consecutive sections are placed contiguously in the index buffer.
+			// Determine the first index based on the index range of the previous rendering section.
+			if( LODSectionIndex == 0 )
+			{
+				StaticMeshSection.FirstIndex = 0;
+			}
+			else
+			{
+				FStaticMeshSection& PreviousStaticMeshSection = StaticMeshLOD.Sections[ LODSectionIndex - 1 ];
+				const FPolygonGroupID PreviousPolygonGroupID = GetSectionForRenderingSectionIndex( LODSectionIndex - 1 );
+				check( PreviousPolygonGroupID != FPolygonGroupID::Invalid );
+				StaticMeshSection.FirstIndex = PreviousStaticMeshSection.FirstIndex + RenderingPolygonGroups[ PreviousPolygonGroupID.GetValue() ].MaxTriangles * 3;
+
+				// @todo mesheditor: if this check is valid, we can dispense with the above and just set StaticMeshSection.FirstIndex according to the current length of the idnex buffer.
+				check( StaticMeshLOD.IndexBuffer.GetNumIndices() == StaticMeshSection.FirstIndex );
+			}
+
+			// Fill in the remaining rendering section properties.
+			StaticMeshSection.NumTriangles = 0;
+			StaticMeshSection.MinVertexIndex = 0;
+			StaticMeshSection.MaxVertexIndex = 0;
+			StaticMeshSection.bEnableCollision = PolygonGroup.bEnableCollision;
+			StaticMeshSection.bCastShadow = PolygonGroup.bCastShadow;
+
+			check( GetStaticMeshMaterialIndex( StaticMesh, PolygonGroup.Material ) == INDEX_NONE );
+			const int32 MaterialIndex = StaticMesh->StaticMaterials.Add( FStaticMaterial( PolygonGroup.Material ) );
+			StaticMeshSection.MaterialIndex = MaterialIndex;
+		}
+
+		// Insert the rendering polygon group for keeping track of these index buffer properties
+		RenderingPolygonGroups.Insert( PolygonGroupID.GetValue(), FRenderingPolygonGroup() );
+		FRenderingPolygonGroup& RenderingPolygonGroup = RenderingPolygonGroups[ PolygonGroupID.GetValue() ];
+
+		RenderingPolygonGroup.RenderingSectionIndex = LODSectionIndex;
+		RenderingPolygonGroup.MaxTriangles = 0;
+	}
+}
+
+
+void UEditableStaticMeshAdapter::OnDeletePolygonGroups( const UEditableMesh* EditableMesh, const TArray<FPolygonGroupID>& PolygonGroupIDs )
+{
+	for( const FPolygonGroupID PolygonGroupID : PolygonGroupIDs )
+	{
+		const FMeshPolygonGroup& PolygonGroup = EditableMesh->PolygonGroups[ PolygonGroupID.GetValue() ];
+		FRenderingPolygonGroup& RenderingPolygonGroup = RenderingPolygonGroups[ PolygonGroupID.GetValue() ];
+
+		// Remove material slot associated with section
+		// @todo mesheditor: can more than one section share a material? Mesh editor currently assumes not, but this will break anything which does
+		const int32 MaterialIndex = GetStaticMeshMaterialIndex( StaticMesh, PolygonGroup.Material );
+		StaticMesh->StaticMaterials.RemoveAt( MaterialIndex );
+
+		// Adjust rendering indices held by sections: any index above the one we just deleted now needs to be decremented.
+		const uint32 RenderingSectionIndex = RenderingPolygonGroup.RenderingSectionIndex;
+
+		for( FRenderingPolygonGroup& PolygonGroupToAdjust : RenderingPolygonGroups )
+		{
+			if( PolygonGroupToAdjust.RenderingSectionIndex > RenderingSectionIndex )
+			{
+				PolygonGroupToAdjust.RenderingSectionIndex--;
+			}
+		}
+
+		if( !EditableMesh->IsPreviewingSubdivisions() )
+		{
+			FStaticMeshLODResources& StaticMeshLOD = GetStaticMeshLOD();
+			const uint32 FirstIndex = StaticMeshLOD.Sections[ RenderingSectionIndex ].FirstIndex;
+
+			// Get current number of triangles allocated for this section
+			const int32 MaxTriangles = RenderingPolygonGroup.MaxTriangles;
+
+			// Remove indices from this poisition in the index buffer
+			StaticMeshLOD.IndexBuffer.RemoveIndicesAt( FirstIndex, MaxTriangles * 3 );
+
+			// Adjust first index for all subsequent render sections to account for the indices just removed.
+			// It is guaranteed that index buffer indices are in the same order as the rendering sections.
+			const uint32 NumRenderingSections = StaticMeshLOD.Sections.Num();
+			for( uint32 Index = RenderingSectionIndex + 1; Index < NumRenderingSections; ++Index )
+			{
+				check( StaticMeshLOD.Sections[ Index ].FirstIndex >= FirstIndex );
+				StaticMeshLOD.Sections[ Index ].FirstIndex -= MaxTriangles * 3;
+			}
+
+			// Adjust material indices for any sections to account for the fact that one has been removed
+			for( uint32 Index = 0; Index < NumRenderingSections; ++Index )
+			{
+				FStaticMeshSection& StaticMeshSection = StaticMeshLOD.Sections[ Index ];
+				if( StaticMeshSection.MaterialIndex > MaterialIndex )
+				{
+					StaticMeshSection.MaterialIndex--;
+				}
+			}
+
+			StaticMeshLOD.Sections.RemoveAt( RenderingSectionIndex );
+		}
+
+		// Remove the rendering polygon group from the sparse array
+		RenderingPolygonGroups.RemoveAt( PolygonGroupID.GetValue() );
+	}
+}
+
+
+FPolygonGroupID UEditableStaticMeshAdapter::GetSectionForRenderingSectionIndex( const int32 RenderingSectionIndex ) const
+{
+	for( auto It = RenderingPolygonGroups.CreateConstIterator(); It; ++It )
+	{
+		if( It->RenderingSectionIndex == RenderingSectionIndex )
+		{
+			return FPolygonGroupID( It.GetIndex() );
+		}
+	}
+
+	return FPolygonGroupID::Invalid;
+}
