Copying //Tasks/UE4/Dev-Geometry-Refactor to Dev-Geometry (//UE4/Dev-Geometry)

diff --git a/Engine/Source/Runtime/MeshEditingRuntime/EditableStaticMeshAdapter.h b/Engine/Source/Runtime/MeshEditingRuntime/EditableStaticMeshAdapter.h
new file mode 100644
--- /dev/null
+++ b/Engine/Source/Runtime/MeshEditingRuntime/EditableStaticMeshAdapter.h
@@ -0,0 +1,202 @@
+// Copyright 1998-2015 Epic Games, Inc. All Rights Reserved.
+
+#pragma once
+
+#include "CoreMinimal.h"
+#include "EditableMeshAdapter.h"
+#include "EditableMeshCustomVersion.h"
+#include "StaticMeshResources.h"
+#include "EditableStaticMeshAdapter.generated.h"
+
+
+USTRUCT()
+struct FTriangleID : public FElementID
+{
+	GENERATED_BODY()
+
+	FTriangleID()
+	{
+	}
+
+	explicit FTriangleID( const FElementID InitElementID )
+		: FElementID( InitElementID.GetValue() )
+	{
+	}
+
+	explicit FTriangleID( const uint32 InitIDValue )
+		: FElementID( InitIDValue )
+	{
+	}
+
+	FORCEINLINE friend uint32 GetTypeHash( const FTriangleID& Other )
+	{
+		return GetTypeHash( Other.IDValue );
+	}
+
+	/** Invalid triangle ID */
+	static const FTriangleID Invalid;
+};
+
+
+USTRUCT()
+struct FRenderingPolygon
+{
+	GENERATED_BODY()
+
+	/** Static meshes currently only support triangles.  We'll always triangulate polygons and keep track
+	    of all of the triangles here */
+	// @todo mesheditor: For other mesh formats, we may not need to assume triangles.  For example, modern
+	//		GPUs can render quads.  We still need to consider collision geometry though.
+	// @todo mesheditor: This should probably always be a transient thing, right?  We can re-triangulate on 
+	//		load?  Not sure.  We might want some edit-time control over triangulation, which should be saved/loaded.
+	UPROPERTY()
+	TArray<FTriangleID> TriangulatedPolygonTriangleIndices;
+
+	/** Serializer */
+	friend FArchive& operator<<( FArchive& Ar, FRenderingPolygon& Polygon )
+	{
+		Ar << Polygon.TriangulatedPolygonTriangleIndices;
+		return Ar;
+	}
+};
+
+
+USTRUCT()
+struct FRenderingPolygonGroup
+{
+	GENERATED_BODY()
+
+	/** The rendering section index for this mesh section */
+	UPROPERTY()
+	uint32 RenderingSectionIndex;
+
+	/** Maximum number of triangles which have been reserved in the index buffer */
+	UPROPERTY()
+	int32 MaxTriangles;
+
+	/** Sparse array of triangles, that matches the triangles in the mesh index buffers.  Elements that
+	    aren't allocated will be stored as degenerates in the mesh index buffer. */
+	TSparseArray<FMeshTriangle> Triangles;
+
+
+	/** Converts from an index in our Triangles array to an index of a rendering triangle's first vertex in the rendering mesh's index buffer */
+	inline static FTriangleID RenderingTriangleFirstIndexToTriangleIndex( const FStaticMeshSection& RenderingSection, const uint32 RenderingTriangleFirstIndex )
+	{
+		return FTriangleID( ( RenderingTriangleFirstIndex - RenderingSection.FirstIndex ) / 3 );
+	}
+
+	/** Converts from an index of a rendering triangle's first vertex in the rendering mesh's index buffer to an index in our Triangles array */
+	inline static uint32 TriangleIndexToRenderingTriangleFirstIndex( const FStaticMeshSection& RenderingSection, const FTriangleID TriangleIndex )
+	{
+		return TriangleIndex.GetValue() * 3 + RenderingSection.FirstIndex;
+	}
+
+	/** Serializer */
+	friend FArchive& operator<<( FArchive& Ar, FRenderingPolygonGroup& Section )
+	{
+		Ar << Section.RenderingSectionIndex;
+		Ar << Section.MaxTriangles;	// @todo mesheditor serialization: Should not need to be serialized if we triangulate after load
+		SerializeSparseArray( Ar, Section.Triangles );	// @todo mesheditor serialization: Should not need to be serialized if we triangulate after load
+		return Ar;
+	}
+};
+
+
+
+
+UCLASS(MinimalAPI)
+class UEditableStaticMeshAdapter : public UEditableMeshAdapter
+{
+	GENERATED_BODY()
+
+public:
+
+	/** Default constructor that initializes good defaults for UEditableStaticMeshAdapter */
+	UEditableStaticMeshAdapter();
+
+	virtual void Serialize( FArchive& Ar ) override;
+
+	/** Creates a editable static mesh from the specified component and sub-mesh address */
+	void InitEditableStaticMesh( UEditableMesh* EditableMesh, class UPrimitiveComponent& Component, const FEditableMeshSubMeshAddress& InitSubMeshAddress );
+	MESHEDITINGRUNTIME_API void InitFromBlankStaticMesh( UEditableMesh* EditableMesh, UStaticMesh& InStaticMesh );
+
+	virtual void OnRebuildRenderMeshStart( const UEditableMesh* EditableMesh, const bool bRefreshBounds, const bool bInvalidateLighting ) override;
+	virtual void OnRebuildRenderMesh( const UEditableMesh* EditableMesh ) override;
+	virtual void OnRebuildRenderMeshFinish( const UEditableMesh* EditableMesh, const bool bUpdateCollision ) override;
+	virtual void OnStartModification( const UEditableMesh* EditableMesh, const EMeshModificationType MeshModificationType, const EMeshTopologyChange MeshTopologyChange ) override;
+	virtual void OnEndModification( const UEditableMesh* EditableMesh ) override;
+	virtual void OnReindexElements( const UEditableMesh* EditableMesh, const FElementIDRemappings& Remappings ) override;
+	virtual bool IsCommitted( const UEditableMesh* EditableMesh ) const override;
+	virtual bool IsCommittedAsInstance( const UEditableMesh* EditableMesh ) const override;
+	virtual void OnCommit( UEditableMesh* EditableMesh ) override;
+	virtual UEditableMesh* OnCommitInstance( UEditableMesh* EditableMesh, UPrimitiveComponent* ComponentToInstanceTo ) override;
+	virtual void OnRevert( UEditableMesh* EditableMesh ) override;
+	virtual UEditableMesh* OnRevertInstance( UEditableMesh* EditableMesh ) override;
+	virtual void OnPropagateInstanceChanges( UEditableMesh* EditableMesh ) override;
+
+	virtual void OnDeleteVertexInstances( const UEditableMesh* EditableMesh, const TArray<FVertexInstanceID>& VertexInstanceIDs ) override;
+	virtual void OnDeleteOrphanVertices( const UEditableMesh* EditableMesh, const TArray<FVertexID>& VertexIDs ) override;
+	virtual void OnCreateEmptyVertexRange( const UEditableMesh* EditableMesh, const TArray<FVertexID>& VertexIDs ) override;
+	virtual void OnCreateVertices( const UEditableMesh* EditableMesh, const TArray<FVertexID>& VertexIDs ) override;
+	virtual void OnCreateVertexInstances( const UEditableMesh* EditableMesh, const TArray<FVertexInstanceID>& VertexInstanceIDs ) override;
+	virtual void OnSetVertexAttribute( const UEditableMesh* EditableMesh, const FVertexID VertexID, const FName AttributeName, const int32 AttributeIndex, const FVector4 AttributeValue ) override;
+	virtual void OnSetVertexInstanceAttribute( const UEditableMesh* EditableMesh, const FVertexInstanceID VertexInstanceID, const FName AttributeName, const int32 AttributeIndex, const FVector4 AttributeValue ) override;
+	virtual void OnCreateEdges( const UEditableMesh* EditableMesh, const TArray<FEdgeID>& EdgeIDs ) override;
+	virtual void OnDeleteEdges( const UEditableMesh* EditableMesh, const TArray<FEdgeID>& EdgeIDs ) override;
+	virtual void OnSetEdgeAttribute( const UEditableMesh* EditableMesh, const FEdgeID EdgeID, const FName AttributeName, const int32 AttributeIndex, const FVector4 AttributeValue ) override;
+	virtual void OnCreatePolygons( const UEditableMesh* EditableMesh, const TArray<FPolygonID>& PolygonIDs ) override;
+	virtual void OnDeletePolygons( const UEditableMesh* EditableMesh, const TArray<FPolygonID>& PolygonIDs ) override;
+	virtual void OnChangePolygonVertexInstances( const UEditableMesh* EditableMesh, const TArray<FPolygonID>& PolygonIDs ) override;
+	virtual void OnCreatePolygonGroups( const UEditableMesh* EditableMesh, const TArray<FPolygonGroupID>& PolygonGroupIDs ) override;
+	virtual void OnDeletePolygonGroups( const UEditableMesh* EditableMesh, const TArray<FPolygonGroupID>& PolygonGroupIDs ) override;
+	virtual void OnRetriangulatePolygons( const UEditableMesh* EditableMesh, const TArray<FPolygonID>& PolygonIDs ) override;
+
+
+private:
+
+
+	/** Deletes all of a polygon's triangles (including rendering triangles from the index buffer.) */
+	void DeletePolygonTriangles( const UEditableMesh* EditableMesh, const FPolygonID PolygonID );
+
+	/** Helper function to get a specific static mesh LOD structure for a mesh */
+	const struct FStaticMeshLODResources& GetStaticMeshLOD() const;
+	struct FStaticMeshLODResources& GetStaticMeshLOD();
+
+	/** Makes sure our mesh's index buffer is 32-bit, converting if needed */
+	void EnsureIndexBufferIs32Bit();
+
+	/** If any of the specified triangles contain vertex indices greater than 0xffff, updates our mesh's index buffer to be 32-bit if it isn't already */
+	void UpdateIndexBufferFormatIfNeeded( const TArray<FMeshTriangle>& Triangles );
+
+	/** Rebuilds bounds */
+	void UpdateBoundsAndCollision( const UEditableMesh* EditableMesh, const bool bUpdateCollision );
+
+	/** Gets the editable mesh section index which corresponds to the given rendering section index */
+	FPolygonGroupID GetSectionForRenderingSectionIndex( const int32 RenderingSectionIndex ) const;
+
+private:
+
+	/** The static mesh asset we're representing */
+	UPROPERTY()
+	UStaticMesh* StaticMesh;
+
+	UPROPERTY()
+	UStaticMesh* OriginalStaticMesh;
+
+	UPROPERTY()
+	int32 StaticMeshLODIndex;
+
+	/** Size of the extra gaps in the index buffer between different sections (so new triangles have some space to be added into,
+		without requiring the index buffer to be manipulated) */
+	static const uint32 IndexBufferInterSectionGap = 32;
+
+	/** All of the polygons in this mesh */
+	TSparseArray< FRenderingPolygon > RenderingPolygons;
+
+	/** All of the polygon groups in this mesh */
+	TSparseArray< FRenderingPolygonGroup > RenderingPolygonGroups;
+
+
+	/** Used to refresh all components in the scene that may be using a mesh we're editing */
+	TSharedPtr< class FStaticMeshComponentRecreateRenderStateContext > RecreateRenderStateContext;
+};
diff --git a/Engine/Source/Runtime/MeshEditingRuntime/Public/EditableMesh.h b/Engine/Source/Runtime/MeshEditingRuntime/Public/EditableMesh.h
--- a/Engine/Source/Runtime/MeshEditingRuntime/Public/EditableMesh.h
+++ b/Engine/Source/Runtime/MeshEditingRuntime/Public/EditableMesh.h
@@ -5,9 +5,16 @@
 #include "CoreMinimal.h"
 #include "UObject/Object.h"
 #include "EditableMeshTypes.h"
+#include "EditableMeshCustomVersion.h"
 #include "Change.h"		// For TUniquePtr<FChange>
+#include "Logging/LogMacros.h"
+#include "Materials/MaterialInterface.h"
 #include "EditableMesh.generated.h"
 
+class UEditableMeshAdapter;
+
+MESHEDITINGRUNTIME_API DECLARE_LOG_CATEGORY_EXTERN( LogEditableMesh, Log, All );
+//MESHEDITINGRUNTIME_API DECLARE_LOG_CATEGORY_EXTERN( LogEditableMesh, All, All );
 
 // @todo mesheditor: Comment these classes and enums!
 
@@ -23,6 +30,351 @@ namespace OpenSubdiv
 }
 
 
+/**
+* Perform custom serialization for TSparseArray.
+* The default TSparseArray serialization also compacts all the elements, removing the gaps and changing the indices.
+* The indices are significant in editable meshes, hence this is a custom serializer which preserves them.
+*/
+template <typename T>
+void SerializeSparseArray( FArchive& Ar, TSparseArray<T>& Array )
+{
+	if( Ar.CustomVer( FEditableMeshCustomVersion::GUID ) < FEditableMeshCustomVersion::CustomSparseArraySerialization )
+	{
+		Ar << Array;
+	}
+	else
+	{
+		Array.CountBytes( Ar );
+
+		if( Ar.IsLoading() )
+		{
+			// Load array
+			TBitArray<> AllocatedIndices;
+			Ar << AllocatedIndices;
+
+			Array.Empty( AllocatedIndices.Num() );
+			for( auto It = TConstSetBitIterator<>( AllocatedIndices ); It; ++It )
+			{
+				Array.Insert( It.GetIndex(), T() );
+				Ar << Array[ It.GetIndex() ];
+			}
+		}
+		else
+		{
+			// Save array
+			const int32 MaxIndex = Array.GetMaxIndex();
+
+			// We have to build the TBitArray representing allocated indices by hand, as we don't have access to it from outside TSparseArray.
+			// @todo core: consider replacing TSparseArray serialization with this format.
+			TBitArray<> AllocatedIndices( false, MaxIndex );
+			for( int32 Index = 0; Index < MaxIndex; ++Index )
+			{
+				if( Array.IsAllocated( Index ) )
+				{
+					AllocatedIndices[ Index ] = true;
+				}
+			}
+			Ar << AllocatedIndices;
+
+			for( auto It = Array.CreateIterator(); It; ++It )
+			{
+				Ar << *It;
+			}
+		}
+	}
+}
+
+
+template <typename T, typename ElementIDType>
+void CompactSparseArrayElements( TSparseArray<T>& Array, TSparseArray<ElementIDType>& IndexRemap )
+{
+	static_assert( TIsDerivedFrom<ElementIDType, FElementID>::IsDerived, "Remap array type must be derived from FElementID" );
+
+	static TSparseArray<T> NewArray;
+	NewArray.Empty( Array.Num() );
+
+	IndexRemap.Empty( Array.GetMaxIndex() );
+
+	// Add valid elements into a new contiguous sparse array.  Note non-const iterator so we can move elements.
+	for( TSparseArray<T>::TIterator It( Array ); It; ++It )
+	{
+		const int32 OldElementIndex = It.GetIndex();
+
+		// @todo mesheditor: implement TSparseArray::Add( ElementType&& ) to save this obscure approach
+		const int32 NewElementIndex = NewArray.Add( T() );
+		NewArray[ NewElementIndex ] = MoveTemp( *It );
+
+		// Provide an O(1) lookup from old index to new index, used when patching up vertex references afterwards
+		IndexRemap.Insert( OldElementIndex, ElementIDType( NewElementIndex ) );
+	}
+
+	Array = MoveTemp( NewArray );
+}
+
+
+template <typename T, typename ElementIDType>
+void RemapSparseArrayElements( TSparseArray<T>& Array, const TSparseArray<ElementIDType>& IndexRemap )
+{
+	static_assert( TIsDerivedFrom<ElementIDType, FElementID>::IsDerived, "Remap array type must be derived from FElementID" );
+
+	static TSparseArray<T> NewArray;
+	NewArray.Empty( IndexRemap.GetMaxIndex() );
+
+	// Add valid elements into a new contiguous sparse array.  Note non-const iterator so we can move elements.
+	for( TSparseArray<T>::TIterator It( Array ); It; ++It )
+	{
+		const int32 OldElementIndex = It.GetIndex();
+
+		check( IndexRemap.IsAllocated( OldElementIndex ) );
+		const int32 NewElementIndex = IndexRemap[ OldElementIndex ].GetValue();
+
+		// @todo mesheditor: implement TSparseArray::Insert( ElementType&& ) to save this obscure approach
+		NewArray.Insert( NewElementIndex, T() );
+		NewArray[ NewElementIndex ] = MoveTemp( *It );
+	}
+
+	Array = MoveTemp( NewArray );
+}
+
+
+USTRUCT()
+struct FMeshVertex
+{
+	GENERATED_BODY()
+
+	/** Position of the vertex. */
+	UPROPERTY()
+	FVector VertexPosition;
+
+	/** All of vertex instances which reference this vertex (for split vertex support) */
+	UPROPERTY()
+	TArray<FVertexInstanceID> VertexInstanceIDs;
+
+	/** The edges connected to this vertex */
+	UPROPERTY()
+	TArray<FEdgeID> ConnectedEdgeIDs;
+
+	/** When subdivisions are enabled, this controls how sharp the vertex is, between 0.0 and 1.0. */
+	UPROPERTY()
+	float CornerSharpness;	// @todo mesheditor subdiv: Not really used by static meshes at all.  Only for UEditableMeshes that use subdivision features.  Move elsewhere.?
+
+	/** Serializer */
+	friend FArchive& operator<<( FArchive& Ar, FMeshVertex& Vertex )
+	{
+		Ar << Vertex.VertexPosition;
+		Ar << Vertex.VertexInstanceIDs;
+		Ar << Vertex.ConnectedEdgeIDs;
+		Ar << Vertex.CornerSharpness;
+		return Ar;
+	}
+};
+
+
+USTRUCT()
+struct FMeshVertexInstance
+{
+	GENERATED_BODY()
+
+	/** The vertex this is instancing */
+	UPROPERTY()
+	FVertexID VertexID;
+
+	/** List of connected polygons */
+	UPROPERTY()
+	TArray<FPolygonID> ConnectedPolygons;
+
+	/** UVs for the vertex instance */
+	UPROPERTY()
+	TArray<FVector2D> VertexUVs;
+
+	/** Normal vector */
+	UPROPERTY()
+	FVector Normal;
+	
+	/** Tangent vector */
+	UPROPERTY()
+	FVector Tangent;
+
+	/** Basis determinant sign used to calculate the sense of the binormal */
+	UPROPERTY()
+	float BinormalSign;
+
+	/** Vertex color */
+	UPROPERTY()
+	FLinearColor Color;
+
+	/** Serializer */
+	friend FArchive& operator<<( FArchive& Ar, FMeshVertexInstance& VertexInstance )
+	{
+		Ar << VertexInstance.VertexID;
+		Ar << VertexInstance.ConnectedPolygons;
+		Ar << VertexInstance.VertexUVs;
+		Ar << VertexInstance.Normal;
+		Ar << VertexInstance.Tangent;
+		Ar << VertexInstance.BinormalSign;
+		Ar << VertexInstance.Color;
+		return Ar;
+	}
+};
+
+
+USTRUCT()
+struct FMeshEdge
+{
+	GENERATED_BODY()
+
+	/** IDs of the two editable mesh vertices that make up this edge.  The winding direction is not defined. */
+	UPROPERTY()
+	FVertexID VertexIDs[ 2 ];
+
+	/** The polygons that share this edge.  It's best if there are always only two polygons that share
+	    the edge, and those polygons are facing the same direction */
+	UPROPERTY()
+	TArray<FPolygonID> ConnectedPolygons;
+
+	/** Whether this edge is 'hard' or not, for the purpose of vertex normal and tangent generation */
+	UPROPERTY()
+	bool bIsHardEdge;
+
+	/** When subdivisions are enabled, this controls how sharp the creasing of this edge will be, between 0.0 and 1.0. */
+	UPROPERTY()
+	float CreaseSharpness;
+
+	/** Serializer */
+	friend FArchive& operator<<( FArchive& Ar, FMeshEdge& Edge )
+	{
+		Ar << Edge.VertexIDs[ 0 ];
+		Ar << Edge.VertexIDs[ 1 ];
+		Ar << Edge.ConnectedPolygons;
+		Ar << Edge.bIsHardEdge;
+		Ar << Edge.CreaseSharpness;
+		return Ar;
+	}
+};
+
+
+USTRUCT()
+struct FMeshPolygonContour
+{
+	GENERATED_BODY()
+
+	/** The ordered list of vertex instances which make up the polygon contour. The winding direction is counter-clockwise. */
+	UPROPERTY()
+	TArray<FVertexInstanceID> VertexInstanceIDs;
+
+	/** Serializer */
+	friend FArchive& operator<<( FArchive& Ar, FMeshPolygonContour& Contour )
+	{
+		Ar << Contour.VertexInstanceIDs;
+		return Ar;
+	}
+};
+
+
+USTRUCT()
+struct FMeshTriangle
+{
+	GENERATED_BODY()
+
+	/** The three vertex instances that make up this triangle, ordered counter-clockwise. */
+	UPROPERTY()
+	FVertexInstanceID VertexInstanceIDs[ 3 ];
+
+	/** Serializer */
+	friend FArchive& operator<<( FArchive& Ar, FMeshTriangle& Tri )
+	{
+		Ar << Tri.VertexInstanceIDs[ 0 ];
+		Ar << Tri.VertexInstanceIDs[ 1 ];
+		Ar << Tri.VertexInstanceIDs[ 2 ];
+		return Ar;
+	}
+};
+
+
+USTRUCT()
+struct FMeshPolygon
+{
+	GENERATED_BODY()
+
+	/** The outer boundary edges of this polygon */
+	UPROPERTY()
+	FMeshPolygonContour PerimeterContour;
+
+	/** Optional inner contours of this polygon that define holes inside of the polygon.  For the geometry to
+	    be considered valid, the hole contours should reside within the boundary of the polygon perimeter contour, 
+		and must not overlap each other.  No "nesting" of polygons inside the holes is supported -- those are 
+		simply separate polygons */
+	UPROPERTY()
+	TArray<FMeshPolygonContour> HoleContours;
+
+	/** List of triangles which make up this polygon */
+	UPROPERTY()
+	TArray<FMeshTriangle> Triangles;
+
+	/** The polygon group which contains this polygon */
+	UPROPERTY()
+	FPolygonGroupID PolygonGroupID;
+
+	/** Cached normal */
+	UPROPERTY()
+	FVector PolygonNormal;
+
+	/** Cached tangent */
+	UPROPERTY()
+	FVector PolygonTangent;
+
+	/** Cached binormal */
+	UPROPERTY()
+	FVector PolygonBinormal;
+
+	/** Cached center */
+	UPROPERTY()
+	FVector PolygonCenter;
+
+	/** Serializer */
+	friend FArchive& operator<<( FArchive& Ar, FMeshPolygon& Polygon )
+	{
+		Ar << Polygon.PerimeterContour;
+		Ar << Polygon.HoleContours;
+		Ar << Polygon.Triangles;
+		Ar << Polygon.PolygonGroupID;
+		return Ar;
+	}
+};
+
+
+USTRUCT()
+struct FMeshPolygonGroup
+{
+	GENERATED_BODY()
+
+	/** The material for this mesh section */
+	UPROPERTY()
+	UMaterialInterface* Material;
+
+	/** If true, collision is enabled for this section. */
+	UPROPERTY()
+	bool bEnableCollision;
+
+	/** If true, this section will cast a shadow */
+	UPROPERTY()
+	bool bCastShadow;
+
+	/** All polygons in this group */
+	TArray<FPolygonID> Polygons;
+
+	/** Serializer */
+	friend FArchive& operator<<( FArchive& Ar, FMeshPolygonGroup& PolygonGroup )
+	{
+		Ar << PolygonGroup.Material;
+		Ar << PolygonGroup.bEnableCollision;
+		Ar << PolygonGroup.bCastShadow;
+		Ar << PolygonGroup.Polygons;
+		return Ar;
+	}
+};
+
+
 UENUM( BlueprintType )
 enum class EInsetPolygonsMode : uint8
 {
@@ -43,7 +395,47 @@ enum class ETriangleTessellationMode : uint8
 };
 
 
-UCLASS( abstract, BlueprintType )
+struct FElementIDRemappings
+{
+	TSparseArray<FVertexID> NewVertexIndexLookup;
+	TSparseArray<FVertexInstanceID> NewVertexInstanceIndexLookup;
+	TSparseArray<FEdgeID> NewEdgeIndexLookup;
+	TSparseArray<FPolygonID> NewPolygonIndexLookup;
+	TSparseArray<FPolygonGroupID> NewPolygonGroupIndexLookup;
+
+	FVertexID GetRemappedVertexID( FVertexID VertexID ) const
+	{
+		checkSlow( NewVertexIndexLookup.IsAllocated( VertexID.GetValue() ) );
+		return NewVertexIndexLookup[ VertexID.GetValue() ];
+	}
+
+	FVertexInstanceID GetRemappedVertexInstanceID( FVertexInstanceID VertexInstanceID ) const
+	{
+		checkSlow( NewVertexInstanceIndexLookup.IsAllocated( VertexInstanceID.GetValue() ) );
+		return NewVertexInstanceIndexLookup[ VertexInstanceID.GetValue() ];
+	}
+
+	FEdgeID GetRemappedEdgeID( FEdgeID EdgeID ) const
+	{
+		checkSlow( NewEdgeIndexLookup.IsAllocated( EdgeID.GetValue() ) );
+		return NewEdgeIndexLookup[ EdgeID.GetValue() ];
+	}
+
+	FPolygonID GetRemappedPolygonID( FPolygonID PolygonID ) const
+	{
+		checkSlow( NewPolygonIndexLookup.IsAllocated( PolygonID.GetValue() ) );
+		return NewPolygonIndexLookup[ PolygonID.GetValue() ];
+	}
+
+	FPolygonGroupID GetRemappedPolygonGroupID( FPolygonGroupID PolygonGroupID ) const
+	{
+		check( NewPolygonGroupIndexLookup.IsAllocated( PolygonGroupID.GetValue() ) );
+		return NewPolygonGroupIndexLookup[ PolygonGroupID.GetValue() ];
+	}
+};
+
+
+UCLASS( BlueprintType )
 class MESHEDITINGRUNTIME_API UEditableMesh : public UObject
 {
 	GENERATED_BODY()
@@ -57,74 +449,98 @@ public:
 	virtual void Serialize( FArchive& Ar ) override;
 	virtual void PostLoad() override;
 
-	UFUNCTION( BlueprintCallable, Category="Editable Mesh" ) virtual void RebuildRenderMesh() PURE_VIRTUAL(,);
-	UFUNCTION( BlueprintCallable, Category="Editable Mesh" ) virtual void StartModification( const EMeshModificationType MeshModificationType, const EMeshTopologyChange MeshTopologyChange ) PURE_VIRTUAL(,);
-	UFUNCTION( BlueprintCallable, Category="Editable Mesh" ) virtual void EndModification( const bool bFromUndo = false ) PURE_VIRTUAL(,);
-	UFUNCTION( BlueprintPure, Category="Editable Mesh" ) virtual bool IsCommitted() const PURE_VIRTUAL(, return false;);
-	UFUNCTION( BlueprintPure, Category="Editable Mesh" ) virtual bool IsCommittedAsInstance() const PURE_VIRTUAL(, return false;);
-	UFUNCTION( BlueprintCallable, Category="Editable Mesh" ) virtual void Commit() PURE_VIRTUAL(,);
-	UFUNCTION( BlueprintCallable, Category="Editable Mesh" ) virtual UEditableMesh* CommitInstance( UPrimitiveComponent* ComponentToInstanceTo ) PURE_VIRTUAL(,return nullptr;);
-	UFUNCTION( BlueprintCallable, Category="Editable Mesh" ) virtual void Revert() PURE_VIRTUAL(,);
-	UFUNCTION( BlueprintCallable, Category="Editable Mesh" ) virtual UEditableMesh* RevertInstance() PURE_VIRTUAL(,return nullptr;);
-	UFUNCTION( BlueprintCallable, Category="Editable Mesh" ) virtual void PropagateInstanceChanges() PURE_VIRTUAL(,);
-
-	UFUNCTION( BlueprintPure, Category="Editable Mesh" ) virtual int32 GetRenderingVertexCount() const PURE_VIRTUAL(, return 0;);
-	UFUNCTION( BlueprintPure, Category="Editable Mesh" ) virtual int32 GetRenderingVertexArraySize() const PURE_VIRTUAL(, return 0;);
-
-	UFUNCTION( BlueprintPure, Category="Editable Mesh" ) virtual int32 GetVertexCount() const PURE_VIRTUAL(,return 0;);
-	UFUNCTION( BlueprintPure, Category="Editable Mesh" ) virtual int32 GetVertexArraySize() const PURE_VIRTUAL(,return 0;);
-	UFUNCTION( BlueprintPure, Category="Editable Mesh" ) virtual bool IsValidVertex( const FVertexID VertexID ) const PURE_VIRTUAL(,return false;);
-	UFUNCTION( BlueprintPure, Category="Editable Mesh" ) virtual FVector4 GetVertexAttribute( const FVertexID VertexID, const FName AttributeName, const int32 AttributeIndex ) const PURE_VIRTUAL(,return FVector4(0.0f););
-	UFUNCTION( BlueprintPure, Category="Editable Mesh" ) virtual int32 GetVertexConnectedEdgeCount( const FVertexID VertexID ) const PURE_VIRTUAL(,return 0;);
-	UFUNCTION( BlueprintPure, Category="Editable Mesh" ) virtual FEdgeID GetVertexConnectedEdge( const FVertexID VertexID, const int32 ConnectedEdgeNumber ) const PURE_VIRTUAL(,return FEdgeID::Invalid;);
-
-	UFUNCTION( BlueprintPure, Category="Editable Mesh" ) virtual FVector4 GetEdgeAttribute( const FEdgeID EdgeID, const FName AttributeName, const int32 AttributeIndex ) const PURE_VIRTUAL(,return FVector4(0.0f););
-	UFUNCTION( BlueprintPure, Category="Editable Mesh" ) virtual int32 GetEdgeCount() const PURE_VIRTUAL(,return 0;);
-	UFUNCTION( BlueprintPure, Category="Editable Mesh" ) virtual int32 GetEdgeArraySize() const PURE_VIRTUAL(,return 0;);
-	UFUNCTION( BlueprintPure, Category="Editable Mesh" ) virtual bool IsValidEdge( const FEdgeID EdgeID ) const PURE_VIRTUAL(,return false;);
-	UFUNCTION( BlueprintPure, Category="Editable Mesh" ) virtual FVertexID GetEdgeVertex( const FEdgeID EdgeID, const int32 EdgeVertexNumber ) const PURE_VIRTUAL(,return FVertexID::Invalid;);
-	UFUNCTION( BlueprintPure, Category="Editable Mesh" ) virtual int32 GetEdgeConnectedPolygonCount( const FEdgeID EdgeID ) const PURE_VIRTUAL(,return 0;);
-	UFUNCTION( BlueprintPure, Category="Editable Mesh" ) virtual FPolygonRef GetEdgeConnectedPolygon( const FEdgeID EdgeID, const int32 ConnectedPolygonNumber ) const PURE_VIRTUAL( , return FPolygonRef(););
-
-	UFUNCTION( BlueprintPure, Category="Editable Mesh" ) virtual int32 GetSectionCount() const PURE_VIRTUAL(,return 0;);
-	UFUNCTION( BlueprintPure, Category="Editable Mesh" ) virtual int32 GetSectionArraySize() const PURE_VIRTUAL(,return 0;);
-	UFUNCTION( BlueprintPure, Category="Editable Mesh" ) virtual bool IsValidSection( const FSectionID SectionID ) const PURE_VIRTUAL(,return false;);
-
-	UFUNCTION( BlueprintPure, Category="Editable Mesh" ) virtual int32 GetPolygonCount( const FSectionID SectionID ) const PURE_VIRTUAL(,return 0;);
-	UFUNCTION( BlueprintPure, Category="Editable Mesh" ) virtual int32 GetPolygonArraySize( const FSectionID SectionID ) const PURE_VIRTUAL(,return 0;);
-	UFUNCTION( BlueprintPure, Category="Editable Mesh" ) virtual bool IsValidPolygon( const FPolygonRef PolygonRef ) const PURE_VIRTUAL(,return false;);
-	UFUNCTION( BlueprintPure, Category="Editable Mesh" ) virtual int32 GetTriangleCount( const FSectionID SectionID ) const PURE_VIRTUAL(,return 0;);
-	UFUNCTION( BlueprintPure, Category="Editable Mesh" ) virtual int32 GetTriangleArraySize( const FSectionID SectionID ) const PURE_VIRTUAL(,return 0;);
-	UFUNCTION( BlueprintPure, Category="Editable Mesh" ) virtual int32 GetPolygonPerimeterVertexCount( const FPolygonRef PolygonRef ) const PURE_VIRTUAL(,return 0;);
-	UFUNCTION( BlueprintPure, Category="Editable Mesh" ) virtual FVertexID GetPolygonPerimeterVertex( const FPolygonRef PolygonRef, const int32 PolygonVertexNumber ) const PURE_VIRTUAL(,return FVertexID::Invalid;);
-	UFUNCTION( BlueprintPure, Category="Editable Mesh" ) virtual FVector4 GetPolygonPerimeterVertexAttribute( const FPolygonRef PolygonRef, const int32 PolygonVertexNumber, const FName AttributeName, const int32 AttributeIndex ) const PURE_VIRTUAL(,return FVector4(0.0f););
-	UFUNCTION( BlueprintPure, Category="Editable Mesh" ) virtual int32 GetPolygonHoleCount( const FPolygonRef PolygonRef ) const PURE_VIRTUAL(,return 0;);
-	UFUNCTION( BlueprintPure, Category="Editable Mesh" ) virtual int32 GetPolygonHoleVertexCount( const FPolygonRef PolygonRef, const int32 HoleNumber ) const PURE_VIRTUAL(,return 0;);
-	UFUNCTION( BlueprintPure, Category="Editable Mesh" ) virtual FVertexID GetPolygonHoleVertex( const FPolygonRef PolygonRef, const int32 HoleNumber, const int32 PolygonVertexNumber ) const PURE_VIRTUAL(,return FVertexID::Invalid;);
-	UFUNCTION( BlueprintPure, Category="Editable Mesh" ) virtual FVector4 GetPolygonHoleVertexAttribute( const FPolygonRef PolygonRef, const int32 HoleNumber, const int32 PolygonVertexNumber, const FName AttributeName, const int32 AttributeIndex ) const PURE_VIRTUAL(,return FVector4(0.0f););
-	UFUNCTION( BlueprintPure, Category="Editable Mesh" ) virtual int32 GetPolygonTriangulatedTriangleCount( const FPolygonRef PolygonRef ) const PURE_VIRTUAL(,return 0;);
-	UFUNCTION( BlueprintPure, Category="Editable Mesh" ) virtual FVector GetPolygonTriangulatedTriangleVertexPosition( const FPolygonRef PolygonRef, const int32 PolygonTriangleNumber, const int32 TriangleVertexNumber ) const PURE_VIRTUAL(,return FVector(0.0f););
-
-	UFUNCTION( BlueprintCallable, Category="Editable Mesh" ) virtual void RetriangulatePolygons( const TArray<FPolygonRef>& PolygonRefs, const bool bOnlyOnUndo ) PURE_VIRTUAL(,);
+	/** Compacts mesh element arrays to remove gaps, and fixes up referenced IDs */
+	void Compact();
+
+	/** Remaps mesh element arrays according to the provided remappings, in order to undo a compact operation */
+	void Uncompact( const FElementIDRemappings& Remappings );
+
+	UFUNCTION( BlueprintCallable, Category="Editable Mesh" ) void RebuildRenderMesh();
+	UFUNCTION( BlueprintCallable, Category="Editable Mesh" ) void StartModification( const EMeshModificationType MeshModificationType, const EMeshTopologyChange MeshTopologyChange );
+	UFUNCTION( BlueprintCallable, Category="Editable Mesh" ) void EndModification( const bool bFromUndo = false );
+	UFUNCTION( BlueprintPure, Category="Editable Mesh" ) bool IsCommitted() const;
+	UFUNCTION( BlueprintPure, Category="Editable Mesh" ) bool IsCommittedAsInstance() const;
+	UFUNCTION( BlueprintCallable, Category="Editable Mesh" ) void Commit();
+	UFUNCTION( BlueprintCallable, Category="Editable Mesh" ) UEditableMesh* CommitInstance( UPrimitiveComponent* ComponentToInstanceTo );
+	UFUNCTION( BlueprintCallable, Category="Editable Mesh" ) void Revert();
+	UFUNCTION( BlueprintCallable, Category="Editable Mesh" ) UEditableMesh* RevertInstance();
+	UFUNCTION( BlueprintCallable, Category="Editable Mesh" ) void PropagateInstanceChanges();
+
+	UFUNCTION( BlueprintPure, Category="Editable Mesh" ) int32 GetVertexCount() const;
+	UFUNCTION( BlueprintPure, Category="Editable Mesh" ) int32 GetVertexArraySize() const;
+	UFUNCTION( BlueprintPure, Category="Editable Mesh" ) bool IsValidVertex( const FVertexID VertexID ) const;
+	UFUNCTION( BlueprintPure, Category="Editable Mesh" ) FVector4 GetVertexAttribute( const FVertexID VertexID, const FName AttributeName, const int32 AttributeIndex ) const;
+	UFUNCTION( BlueprintPure, Category="Editable Mesh" ) int32 GetVertexConnectedEdgeCount( const FVertexID VertexID ) const;
+	UFUNCTION( BlueprintPure, Category="Editable Mesh" ) FEdgeID GetVertexConnectedEdge( const FVertexID VertexID, const int32 ConnectedEdgeNumber ) const;
+
+	UFUNCTION( BlueprintPure, Category="Editable Mesh" ) int32 GetVertexInstanceCount() const;
+	UFUNCTION( BlueprintPure, Category="Editable Mesh" ) int32 GetVertexInstanceArraySize() const;
+	UFUNCTION( BlueprintPure, Category="Editable Mesh" ) FVector4 GetVertexInstanceAttribute( const FVertexInstanceID VertexInstanceID, const FName AttributeName, const int32 AttributeIndex ) const;
+
+	UFUNCTION( BlueprintPure, Category="Editable Mesh" ) FVector4 GetEdgeAttribute( const FEdgeID EdgeID, const FName AttributeName, const int32 AttributeIndex ) const;
+	UFUNCTION( BlueprintPure, Category="Editable Mesh" ) int32 GetEdgeCount() const;
+	UFUNCTION( BlueprintPure, Category="Editable Mesh" ) int32 GetEdgeArraySize() const;
+	UFUNCTION( BlueprintPure, Category="Editable Mesh" ) bool IsValidEdge( const FEdgeID EdgeID ) const;
+	UFUNCTION( BlueprintPure, Category="Editable Mesh" ) FVertexID GetEdgeVertex( const FEdgeID EdgeID, const int32 EdgeVertexNumber ) const;
+	UFUNCTION( BlueprintPure, Category="Editable Mesh" ) int32 GetEdgeConnectedPolygonCount( const FEdgeID EdgeID ) const;
+	UFUNCTION( BlueprintPure, Category="Editable Mesh" ) FPolygonID GetEdgeConnectedPolygon( const FEdgeID EdgeID, const int32 ConnectedPolygonNumber ) const;
+
+	UFUNCTION( BlueprintPure, Category="Editable Mesh" ) int32 GetPolygonGroupCount() const;
+	UFUNCTION( BlueprintPure, Category="Editable Mesh" ) int32 GetPolygonGroupArraySize() const;
+	UFUNCTION( BlueprintPure, Category="Editable Mesh" ) bool IsValidPolygonGroup( const FPolygonGroupID PolygonGroupID ) const;
+	UFUNCTION( BlueprintPure, Category="Editable Mesh" ) int32 GetPolygonCountInGroup( const FPolygonGroupID PolygonGroupID ) const;
+	UFUNCTION( BlueprintPure, Category="Editable Mesh" ) FPolygonID GetPolygonInGroup( const FPolygonGroupID PolygonGroupID, const int32 PolygonNumber ) const;
+
+	UFUNCTION( BlueprintPure, Category="Editable Mesh" ) int32 GetPolygonCount() const;
+	UFUNCTION( BlueprintPure, Category="Editable Mesh" ) int32 GetPolygonArraySize() const;
+	UFUNCTION( BlueprintPure, Category="Editable Mesh" ) bool IsValidPolygon( const FPolygonID PolygonID ) const;
+	UFUNCTION( BlueprintPure, Category="Editable Mesh" ) FPolygonGroupID GetGroupForPolygon( const FPolygonID PolygonID ) const;
+
+//	UFUNCTION( BlueprintPure, Category="Editable Mesh" ) virtual int32 GetTriangleCount( const FPolygonGroupID PolygonGroupID ) const PURE_VIRTUAL(,return 0;);
+//	UFUNCTION( BlueprintPure, Category="Editable Mesh" ) virtual int32 GetTriangleArraySize( const FPolygonGroupID PolygonGroupID ) const PURE_VIRTUAL(,return 0;);
+
+	UFUNCTION( BlueprintPure, Category="Editable Mesh" ) int32 GetPolygonPerimeterVertexCount( const FPolygonID PolygonID ) const;
+	UFUNCTION( BlueprintPure, Category="Editable Mesh" ) FVertexID GetPolygonPerimeterVertex( const FPolygonID PolygonID, const int32 PolygonVertexNumber ) const;
+	UFUNCTION( BlueprintPure, Category="Editable Mesh" ) FVertexInstanceID GetPolygonPerimeterVertexInstance( const FPolygonID PolygonID, const int32 PolygonVertexNumber ) const;
+	UFUNCTION( BlueprintPure, Category="Editable Mesh" ) FVector4 GetPolygonPerimeterVertexAttribute( const FPolygonID PolygonID, const int32 PolygonVertexNumber, const FName AttributeName, const int32 AttributeIndex ) const;
+	UFUNCTION( BlueprintPure, Category="Editable Mesh" ) int32 GetPolygonHoleCount( const FPolygonID PolygonID ) const;
+	UFUNCTION( BlueprintPure, Category="Editable Mesh" ) int32 GetPolygonHoleVertexCount( const FPolygonID PolygonID, const int32 HoleNumber ) const;
+	UFUNCTION( BlueprintPure, Category="Editable Mesh" ) FVertexID GetPolygonHoleVertex( const FPolygonID PolygonID, const int32 HoleNumber, const int32 PolygonVertexNumber ) const;
+	UFUNCTION( BlueprintPure, Category="Editable Mesh" ) FVertexInstanceID GetPolygonHoleVertexInstance( const FPolygonID PolygonID, const int32 HoleNumber, const int32 PolygonVertexNumber ) const;
+	UFUNCTION( BlueprintPure, Category="Editable Mesh" ) FVector4 GetPolygonHoleVertexAttribute( const FPolygonID PolygonID, const int32 HoleNumber, const int32 PolygonVertexNumber, const FName AttributeName, const int32 AttributeIndex ) const;
+
+	UFUNCTION( BlueprintPure, Category="Editable Mesh" ) int32 GetPolygonTriangulatedTriangleCount( const FPolygonID PolygonID ) const;
+	UFUNCTION( BlueprintPure, Category="Editable Mesh" ) FMeshTriangle GetPolygonTriangulatedTriangle( const FPolygonID PolygonID, int32 PolygonTriangleNumber ) const;
+	UFUNCTION( BlueprintPure, Category="Editable Mesh" ) FVector GetPolygonTriangulatedTriangleVertexPosition( const FPolygonID PolygonID, const int32 PolygonTriangleNumber, const int32 TriangleVertexNumber ) const;
 
 protected:
 
-	virtual void DeleteOrphanVertices_Internal( const TArray<FVertexID>& VertexIDsToDelete ) PURE_VIRTUAL( , );
-	virtual void SetVertexAttribute_Internal( const FVertexID VertexID, const FName AttributeName, const int32 AttributeIndex, const FVector4& NewAttributeValue ) PURE_VIRTUAL( , );
-	virtual void CreateEmptyVertexRange_Internal( const int32 NumVerticesToAdd, const TArray<FVertexID>* OverrideVertexIDsForRedo, TArray<FVertexID>& OutNewVertexIDs ) PURE_VIRTUAL(,);
-	virtual void CreateEdge_Internal( const FVertexID VertexIDA, const FVertexID VertexIDB, const TArray<FPolygonRef>& ConnectedPolygons, const FEdgeID OverrideEdgeIDForRedo, FEdgeID& OutNewEdgeID ) PURE_VIRTUAL(,);
-	virtual void SetEdgeAttribute_Internal( const FEdgeID EdgeID, const FName AttributeName, const int32 AttributeIndex, const FVector4& NewAttributeValue ) PURE_VIRTUAL( , );
-	virtual void DeleteEdges_Internal( const TArray<FEdgeID>& EdgeIDsToDelete, const bool bDeleteOrphanedVertices ) PURE_VIRTUAL(,);
-	virtual void SetEdgeVertices_Internal( const FEdgeID EdgeID, const FVertexID NewVertexID0, const FVertexID NewVertexID1 ) PURE_VIRTUAL(,);
-	virtual void CreatePolygon_Internal( const FSectionID SectionID, const TArray<FVertexID>& VertexIDs, const TArray<TArray<FVertexID>>& VertexIDsForEachHole, const FPolygonID OverridePolygonIDForRedo, FPolygonRef& OutNewPolygonRef, TArray<FEdgeID>& OutNewEdgeIDs ) PURE_VIRTUAL( ,);
-	virtual void DeletePolygon_Internal( const FPolygonRef PolygonRef, const bool bDeleteOrphanedEdges, const bool bDeleteOrphanedVertices, const bool bDeleteEmptySections ) PURE_VIRTUAL(,);
-	virtual void SetPolygonPerimeterVertexAttribute_Internal( const FPolygonRef PolygonRef, const int32 PolygonVertexNumber, const FName AttributeName, const int32 AttributeIndex, const FVector4& NewAttributeValue ) PURE_VIRTUAL( , );
-	virtual void SetPolygonHoleVertexAttribute_Internal( const FPolygonRef PolygonRef, const int32 HoleNumber, const int32 PolygonVertexNumber, const FName AttributeName, const int32 AttributeIndex, const FVector4& NewAttributeValue ) PURE_VIRTUAL( , );
-	virtual void InsertPolygonPerimeterVertices_Internal( const FPolygonRef PolygonRef, const int32 InsertBeforeVertexNumber, const TArray<FVertexAndAttributes>& VerticesToInsert ) PURE_VIRTUAL(,);
-	virtual void RemovePolygonPerimeterVertices_Internal( const FPolygonRef PolygonRef, const int32 FirstVertexNumberToRemove, const int32 NumVerticesToRemove ) PURE_VIRTUAL(,);
-	virtual FSectionID GetSectionIDFromMaterial_Internal( UMaterialInterface* Material, bool bCreateNewSectionIfNotFound ) PURE_VIRTUAL(, return FSectionID::Invalid; );
-	virtual FSectionID CreateSection_Internal( const FSectionToCreate& SectionToCreate ) PURE_VIRTUAL(,return FSectionID::Invalid;);
-	virtual void DeleteSection_Internal( const FSectionID SectionID ) PURE_VIRTUAL(,);
+	/** Given a set of index remappings, fixes up references to element IDs */
+	void FixUpElementIDs( const FElementIDRemappings& Remappings );
+
+	FEdgeID GetPolygonContourEdge( const FMeshPolygonContour& Contour, const int32 ContourEdgeNumber, bool& bOutEdgeWindingIsReversedForPolygon ) const;
+	void GetPolygonContourEdges( const FMeshPolygonContour& Contour, TArray<FEdgeID>& OutPolygonContourEdgeIDs ) const;
+	void SetVertexAttribute( const FVertexID VertexID, const FName AttributeName, const int32 AttributeIndex, const FVector4 AttributeValue );
+	void SetVertexInstanceAttribute( const FVertexInstanceID VertexInstanceID, const FName AttributeName, const int32 AttributeIndex, const FVector4 AttributeValue );
+	void SetEdgeAttribute( const FEdgeID EdgeID, const FName AttributeName, const int32 AttributeIndex, const FVector4 AttributeValue, bool bIsUndo );
+	void GetVertexInstanceConnectedPolygonsInSameGroup( const FVertexInstanceID VertexInstanceID, const FPolygonID PolygonID, TArray<FPolygonID>& OutPolygonIDs ) const;
+	void SetEdgeHardness( const FEdgeID EdgeID, const bool bIsHard, const bool bIsUndo );
+	FVertexInstanceID CreateVertexInstanceForContourVertex( const FVertexAndAttributes& ContourVertex, const FPolygonID PolygonID );
+	void CreatePolygonContour( const FPolygonID PolygonID, const TArray<FVertexAndAttributes>& Contour, const EPolygonEdgeHardness PolygonEdgeHardness, TArray<FEdgeID>& OutEdgeIDs, TArray<FVertexInstanceID>& OutVertexInstanceIDs );
+	void BackupPolygonContour( const FMeshPolygonContour& Contour, TArray<FVertexAndAttributes>& OutVerticesAndAttributes );
+	void DeletePolygonContour( const FPolygonID PolygonID, const FMeshPolygonContour& Contour, const bool bDeleteOrphanedEdges, const bool bDeleteOrphanedVertexInstances, TArray<FEdgeID>& OrphanedEdgeIDs, TArray<FVertexInstanceID>& OrphanedVertexInstanceIDs );
+	static const TArray<FName>& GetMergeableVertexInstanceAttributes();
+	void GetMergedNamedVertexInstanceAttributeData( const TArray<FName>& AttributeNames, const FVertexInstanceID VertexInstanceID, const TArray<FMeshElementAttributeData>& AttributeValuesToMerge, TArray<FMeshElementAttributeData>& OutAttributes ) const;
+	void GetConnectedSoftEdges( const FVertexID VertexID, TArray<FEdgeID>& OutConnectedSoftEdges ) const;
+	void GetPolygonsInSameSoftEdgedGroup( const FVertexID VertexInstanceID, const FPolygonID PolygonID, TArray<FPolygonID>& OutPolygonIDs ) const;
+	void GetVertexInstancesInSameSoftEdgedGroup( const FVertexID VertexID, const FPolygonID PolygonID, const bool bPolygonNotYetInitialized, TArray<FVertexInstanceID>& OutVertexInstanceIDs ) const;
+	FVertexInstanceID GetVertexInstanceInPolygonForVertex( const FPolygonID PolygonID, const FVertexID VertexID ) const;
+	void SetPolygonContourVertexAttributes( FMeshPolygonContour& Contour, const FPolygonID PolygonID, const int32 HoleIndex, const TArray<FMeshElementAttributeList>& AttributeLists );
+	bool AreAttributeListsEqual( const TArray<FMeshElementAttributeData>& A, const TArray<FMeshElementAttributeData>& B ) const;
+	void SplitVertexInstanceInPolygons( const FVertexInstanceID VertexInstanceID, const TArray<FPolygonID>& PolygonIDs );
+	void ReplaceVertexInstanceInPolygons( const FVertexInstanceID OldVertexInstanceID, const FVertexInstanceID NewVertexInstanceID, const TArray<FPolygonID>& PolygonIDs );
+	float GetPolygonCornerAngleForVertex( const FPolygonID PolygonID, const FVertexID VertexID ) const;
+	FPolygonGroupID GetPolygonGroupIDFromMaterial( UMaterialInterface* Material, bool bCreateNewSectionIfNotFound );
 
 public:
 
@@ -189,8 +605,9 @@ public:
 	 */
 
 	UFUNCTION( BlueprintPure, Category="Editable Mesh" ) static const TArray<FName>& GetValidVertexAttributes();
-	UFUNCTION( BlueprintPure, Category="Editable Mesh" ) static const TArray<FName>& GetValidPolygonVertexAttributes();
+	UFUNCTION( BlueprintPure, Category="Editable Mesh" ) static const TArray<FName>& GetValidVertexInstanceAttributes();
 	UFUNCTION( BlueprintPure, Category="Editable Mesh" ) static const TArray<FName>& GetValidEdgeAttributes();
+	UFUNCTION( BlueprintPure, Category="Editable Mesh" ) static const TArray<FName>& GetValidPolygonAttributes();
 	UFUNCTION( BlueprintPure, Category="Editable Mesh" ) static inline FVertexID InvalidVertexID()
 	{
 		return FVertexID::Invalid;
@@ -199,18 +616,14 @@ public:
 	{
 		return FEdgeID::Invalid;
 	}
-	UFUNCTION( BlueprintPure, Category="Editable Mesh" ) static inline FSectionID InvalidSectionID()
+	UFUNCTION( BlueprintPure, Category="Editable Mesh" ) static inline FPolygonGroupID InvalidPolygonGroupID()
 	{
-		return FSectionID::Invalid;
+		return FPolygonGroupID::Invalid;
 	}
 	UFUNCTION( BlueprintPure, Category="Editable Mesh" ) static inline FPolygonID InvalidPolygonID()
 	{
 		return FPolygonID::Invalid;
 	}
-	UFUNCTION( BlueprintPure, Category="Editable Mesh" ) static inline FPolygonRef InvalidPolygonRef()
-	{
-		return FPolygonRef::Invalid;
-	}
 	UFUNCTION( BlueprintPure, Category="Editable Mesh" ) static inline FVertexID MakeVertexID( const int32 VertexIndex )
 	{
 		return FVertexID( VertexIndex );
@@ -219,9 +632,9 @@ public:
 	{
 		return FEdgeID( EdgeIndex );
 	}
-	UFUNCTION( BlueprintPure, Category="Editable Mesh" ) static inline FSectionID MakeSectionID( const int32 SectionIndex )
+	UFUNCTION( BlueprintPure, Category="Editable Mesh" ) static inline FPolygonGroupID MakePolygonGroupID( const int32 PolygonGroupIndex )
 	{
-		return FSectionID( SectionIndex );
+		return FPolygonGroupID( PolygonGroupIndex );
 	}
 	UFUNCTION( BlueprintPure, Category="Editable Mesh" ) static inline FPolygonID MakePolygonID( const int32 PolygonIndex )
 	{
@@ -229,101 +642,111 @@ public:
 	}
 
 	UFUNCTION( BlueprintPure, Category="Editable Mesh" ) int32 GetMaxAttributeIndex( const FName AttributeName ) const;
-	UFUNCTION( BlueprintPure, Category="Editable Mesh" ) FSectionID GetFirstValidSection() const;
-	UFUNCTION( BlueprintPure, Category="Editable Mesh" ) int32 GetTotalPolygonCount() const;
+	UFUNCTION( BlueprintPure, Category="Editable Mesh" ) FPolygonGroupID GetFirstValidPolygonGroup() const;
 	UFUNCTION( BlueprintPure, Category="Editable Mesh" ) int32 GetTextureCoordinateCount() const;
 	UFUNCTION( BlueprintPure, Category="Editable Mesh" ) int32 GetSubdivisionCount() const;
 	UFUNCTION( BlueprintPure, Category="Editable Mesh" ) bool IsPreviewingSubdivisions() const;
 	UFUNCTION( BlueprintPure, Category="Editable Mesh" ) void GetVertexConnectedEdges( const FVertexID VertexID, TArray<FEdgeID>& OutConnectedEdgeIDs ) const;
-	UFUNCTION( BlueprintPure, Category="Editable Mesh" ) void GetVertexConnectedPolygons( const FVertexID VertexID, TArray<FPolygonRef>& OutConnectedPolygonRefs ) const;
+	UFUNCTION( BlueprintPure, Category="Editable Mesh" ) void GetVertexConnectedPolygons( const FVertexID VertexID, TArray<FPolygonID>& OutConnectedPolygonIDs ) const;
+	UFUNCTION( BlueprintPure, Category="Editable Mesh" ) void GetVertexInstanceConnectedPolygons( const FVertexInstanceID VertexInstanceID, TArray<FPolygonID>& OutConnectedPolygonIDs ) const;
+	UFUNCTION( BlueprintPure, Category="Editable Mesh" ) int32 GetVertexInstanceConnectedPolygonCount( const FVertexInstanceID VertexInstanceID ) const;
 	UFUNCTION( BlueprintPure, Category="Editable Mesh" ) void GetVertexAdjacentVertices( const FVertexID VertexID, TArray< FVertexID >& OutAdjacentVertexIDs ) const;
+	UFUNCTION( BlueprintPure, Category="Editable Mesh" ) FEdgeID GetVertexPairEdge( const FVertexID VertexID, const FVertexID NextVertexID, bool& bOutEdgeWindingIsReversed ) const;
 	UFUNCTION( BlueprintPure, Category="Editable Mesh" ) void GetEdgeVertices( const FEdgeID EdgeID, FVertexID& OutEdgeVertexID0, FVertexID& OutEdgeVertexID1 ) const;
-	UFUNCTION( BlueprintPure, Category="Editable Mesh" ) void GetEdgeConnectedPolygons( const FEdgeID EdgeID, TArray<FPolygonRef>& OutConnectedPolygonRefs ) const;
+	UFUNCTION( BlueprintPure, Category="Editable Mesh" ) void GetEdgeConnectedPolygons( const FEdgeID EdgeID, TArray<FPolygonID>& OutConnectedPolygonIDs ) const;
 	UFUNCTION( BlueprintPure, Category="Editable Mesh" ) void GetEdgeLoopElements( const FEdgeID EdgeID, TArray<FEdgeID>& EdgeLoopIDs ) const;
 	UFUNCTION( BlueprintPure, Category="Editable Mesh" ) FEdgeID GetEdgeThatConnectsVertices( const FVertexID VertexID0, const FVertexID VertexID1 ) const;
-	UFUNCTION( BlueprintPure, Category="Editable Mesh" ) int32 GetPolygonPerimeterEdgeCount( const FPolygonRef PolygonRef ) const;
-	UFUNCTION( BlueprintPure, Category="Editable Mesh" ) int32 GetPolygonHoleEdgeCount( const FPolygonRef PolygonRef, const int32 HoleNumber ) const;
-	UFUNCTION( BlueprintPure, Category="Editable Mesh" ) void GetPolygonPerimeterVertices( const FPolygonRef PolygonRef, TArray<FVertexID>& OutPolygonPerimeterVertexIDs ) const;
-	UFUNCTION( BlueprintPure, Category="Editable Mesh" ) void GetPolygonHoleVertices( const FPolygonRef PolygonRef, const int32 HoleNumber, TArray<FVertexID>& OutHoleVertexIDs ) const;
-	UFUNCTION( BlueprintPure, Category="Editable Mesh" ) FEdgeID GetPolygonPerimeterEdge( const FPolygonRef PolygonRef, const int32 PerimeterEdgeNumber, bool& bOutEdgeWindingIsReversedForPolygon ) const;
-	UFUNCTION( BlueprintPure, Category="Editable Mesh" ) FEdgeID GetPolygonHoleEdge( const FPolygonRef PolygonRef, const int32 HoleNumber, const int32 HoleEdgeNumber ) const;
-	UFUNCTION( BlueprintPure, Category="Editable Mesh" ) void GetPolygonPerimeterEdges( const FPolygonRef PolygonRef, TArray<FEdgeID>& OutPolygonPerimeterEdgeIDs ) const;
-	UFUNCTION( BlueprintPure, Category="Editable Mesh" ) void GetPolygonHoleEdges( const FPolygonRef PolygonRef, const int32 HoleNumber, TArray<FEdgeID>& OutHoleEdgeIDs ) const;
-	UFUNCTION( BlueprintPure, Category="Editable Mesh" ) void GetPolygonAdjacentPolygons( const FPolygonRef PolygonRef, TArray<FPolygonRef>& OutAdjacentPolygons ) const;
-	UFUNCTION( BlueprintPure, Category="Editable Mesh" ) int32 FindPolygonPerimeterVertexNumberForVertex( const FPolygonRef PolygonRef, const FVertexID VertexID ) const;
-	UFUNCTION( BlueprintPure, Category="Editable Mesh" ) int32 FindPolygonHoleVertexNumberForVertex( const FPolygonRef PolygonRef, const int32 HoleNumber, const FVertexID VertexID ) const;
-	UFUNCTION( BlueprintPure, Category="Editable Mesh" ) int32 FindPolygonPerimeterEdgeNumberForVertices( const FPolygonRef PolygonRef, const FVertexID EdgeVertexID0, const FVertexID EdgeVertexID1 ) const;
-	UFUNCTION( BlueprintPure, Category="Editable Mesh" ) int32 FindPolygonHoleEdgeNumberForVertices( const FPolygonRef PolygonRef, const int32 HoleNumber, const FVertexID EdgeVertexID0, const FVertexID EdgeVertexID1 ) const;
+	UFUNCTION( BlueprintPure, Category="Editable Mesh" ) int32 GetPolygonPerimeterEdgeCount( const FPolygonID PolygonID ) const;
+	UFUNCTION( BlueprintPure, Category="Editable Mesh" ) int32 GetPolygonHoleEdgeCount( const FPolygonID PolygonID, const int32 HoleNumber ) const;
+	UFUNCTION( BlueprintPure, Category="Editable Mesh" ) void GetPolygonPerimeterVertices( const FPolygonID PolygonID, TArray<FVertexID>& OutPolygonPerimeterVertexIDs ) const;
+	UFUNCTION( BlueprintPure, Category="Editable Mesh" ) void GetPolygonPerimeterVertexInstances( const FPolygonID PolygonID, TArray<FVertexInstanceID>& OutPolygonPerimeterVertexInstanceIDs ) const;
+	UFUNCTION( BlueprintPure, Category="Editable Mesh" ) void GetPolygonHoleVertices( const FPolygonID PolygonID, const int32 HoleNumber, TArray<FVertexID>& OutHoleVertexIDs ) const;
+	UFUNCTION( BlueprintPure, Category="Editable Mesh" ) void GetPolygonHoleVertexInstances( const FPolygonID PolygonID, const int32 HoleNumber, TArray<FVertexInstanceID>& OutHoleVertexInstanceIDs ) const;
+	UFUNCTION( BlueprintPure, Category="Editable Mesh" ) FEdgeID GetPolygonPerimeterEdge( const FPolygonID PolygonID, const int32 PerimeterEdgeNumber, bool& bOutEdgeWindingIsReversedForPolygon ) const;
+	UFUNCTION( BlueprintPure, Category="Editable Mesh" ) FEdgeID GetPolygonHoleEdge( const FPolygonID PolygonID, const int32 HoleNumber, const int32 HoleEdgeNumber ) const;
+	UFUNCTION( BlueprintPure, Category="Editable Mesh" ) void GetPolygonPerimeterEdges( const FPolygonID PolygonID, TArray<FEdgeID>& OutPolygonPerimeterEdgeIDs ) const;
+	UFUNCTION( BlueprintPure, Category="Editable Mesh" ) void GetPolygonHoleEdges( const FPolygonID PolygonID, const int32 HoleNumber, TArray<FEdgeID>& OutHoleEdgeIDs ) const;
+	UFUNCTION( BlueprintPure, Category="Editable Mesh" ) void GetPolygonAdjacentPolygons( const FPolygonID PolygonID, TArray<FPolygonID>& OutAdjacentPolygons ) const;
+	UFUNCTION( BlueprintPure, Category="Editable Mesh" ) int32 FindPolygonPerimeterVertexNumberForVertex( const FPolygonID PolygonID, const FVertexID VertexID ) const;
+	UFUNCTION( BlueprintPure, Category="Editable Mesh" ) int32 FindPolygonHoleVertexNumberForVertex( const FPolygonID PolygonID, const int32 HoleNumber, const FVertexID VertexID ) const;
+	UFUNCTION( BlueprintPure, Category="Editable Mesh" ) int32 FindPolygonPerimeterEdgeNumberForVertices( const FPolygonID PolygonID, const FVertexID EdgeVertexID0, const FVertexID EdgeVertexID1 ) const;
+	UFUNCTION( BlueprintPure, Category="Editable Mesh" ) int32 FindPolygonHoleEdgeNumberForVertices( const FPolygonID PolygonID, const int32 HoleNumber, const FVertexID EdgeVertexID0, const FVertexID EdgeVertexID1 ) const;
 	UFUNCTION( BlueprintPure, Category="Editable Mesh" ) FBox ComputeBoundingBox() const;
 	UFUNCTION( BlueprintPure, Category="Editable Mesh" ) FBoxSphereBounds ComputeBoundingBoxAndSphere() const;
-	UFUNCTION( BlueprintPure, Category="Editable Mesh" ) FVector ComputePolygonCenter( const FPolygonRef PolygonRef ) const;
-	UFUNCTION( BlueprintPure, Category="Editable Mesh" ) FPlane ComputePolygonPlane( const FPolygonRef PolygonRef ) const;
-	UFUNCTION( BlueprintPure, Category="Editable Mesh" ) FVector ComputePolygonNormal( const FPolygonRef PolygonRef ) const;
-	UFUNCTION( BlueprintPure, Category="Editable Mesh" ) FVector ComputePolygonPerimeterVertexNormal( const FPolygonRef PolygonRef, const int32 PolygonVertexNumber ) const;
-	void ComputePolygonTriangulation( const FPolygonRef PolygonRef, TArray<int32>& OutPerimeterVertexNumbersForTriangles ) const;
+	UFUNCTION( BlueprintPure, Category="Editable Mesh" ) FVector ComputePolygonCenter( const FPolygonID PolygonID ) const;
+	UFUNCTION( BlueprintPure, Category="Editable Mesh" ) FPlane ComputePolygonPlane( const FPolygonID PolygonID ) const;
+	UFUNCTION( BlueprintPure, Category="Editable Mesh" ) FVector ComputePolygonNormal( const FPolygonID PolygonID ) const;
 	UFUNCTION( BlueprintPure, Category="Editable Mesh" ) const FSubdivisionLimitData& GetSubdivisionLimitData() const;
-	UFUNCTION( BlueprintPure, Category="Editable Mesh" ) void ComputePolygonTriangulation( const FPolygonRef PolygonRef, TArray<FVertexID>& OutTrianglesVertexIDs ) const;
-	UFUNCTION( BlueprintPure, Category="Editable Mesh" ) bool ComputeBarycentricWeightForPointOnPolygon( const FPolygonRef PolygonRef, const FVector PointOnPolygon, TArray<int32>& PerimeterVertexIndices, FVector& TriangleVertexWeights ) const;
-	UFUNCTION( BlueprintPure, Category="Editable Mesh" ) void ComputeTextureCoordinatesForPointOnPolygon( const FPolygonRef PolygonRef, const FVector PointOnPolygon, bool& bOutWereTextureCoordinatesFound, TArray<FVector4>& OutInterpolatedTextureCoordinates );
-	UFUNCTION( BlueprintPure, Category="Editable Mesh" ) void ComputePolygonsSharedEdges( const TArray<FPolygonRef>& PolygonRefs, TArray<FEdgeID>& OutSharedEdgeIDs ) const;
-	UFUNCTION( BlueprintPure, Category="Editable Mesh" ) void FindPolygonLoop( const FEdgeID EdgeID, TArray<FEdgeID>& OutEdgeLoopEdgeIDs, TArray<FEdgeID>& OutFlippedEdgeIDs, TArray<FEdgeID>& OutReversedEdgeIDPathToTake, TArray<FPolygonRef>& OutPolygonRefsToSplit ) const;
+	UFUNCTION( BlueprintPure, Category="Editable Mesh" ) void ComputePolygonTriangulation( const FPolygonID PolygonID, TArray<FMeshTriangle>& OutTriangles ) const;
+	UFUNCTION( BlueprintPure, Category="Editable Mesh" ) bool ComputeBarycentricWeightForPointOnPolygon( const FPolygonID PolygonID, const FVector PointOnPolygon, FMeshTriangle& OutTriangle, FVector& OutTriangleVertexWeights ) const;
+	UFUNCTION( BlueprintPure, Category="Editable Mesh" ) void ComputeTextureCoordinatesForPointOnPolygon( const FPolygonID PolygonID, const FVector PointOnPolygon, bool& bOutWereTextureCoordinatesFound, TArray<FVector4>& OutInterpolatedTextureCoordinates );
+	UFUNCTION( BlueprintPure, Category="Editable Mesh" ) void ComputePolygonsSharedEdges( const TArray<FPolygonID>& PolygonIDs, TArray<FEdgeID>& OutSharedEdgeIDs ) const;
+	UFUNCTION( BlueprintPure, Category="Editable Mesh" ) void FindPolygonLoop( const FEdgeID EdgeID, TArray<FEdgeID>& OutEdgeLoopEdgeIDs, TArray<FEdgeID>& OutFlippedEdgeIDs, TArray<FEdgeID>& OutReversedEdgeIDPathToTake, TArray<FPolygonID>& OutPolygonIDsToSplit ) const;
 
 
 	// Low poly editing features
 	// @todo mesheditor: Move elsewhere
 	UFUNCTION( BlueprintCallable, Category="Editable Mesh" ) void SetSubdivisionCount( const int32 NewSubdivisionCount );
 	UFUNCTION( BlueprintCallable, Category="Editable Mesh" ) void MoveVertices( const TArray<FVertexToMove>& VerticesToMove );
-	UFUNCTION( BlueprintCallable, Category="Editable Mesh" ) void CreateMissingPolygonPerimeterEdges( const FPolygonRef PolygonRef, TArray<FEdgeID>& OutNewEdgeIDs );
-	UFUNCTION( BlueprintCallable, Category="Editable Mesh" ) void CreateMissingPolygonHoleEdges( const FPolygonRef PolygonRef, const int32 HoleNumber, TArray<FEdgeID>& OutNewEdgeIDs );
+	UFUNCTION( BlueprintCallable, Category="Editable Mesh" ) void CreateMissingPolygonPerimeterEdges( const FPolygonID PolygonID, TArray<FEdgeID>& OutNewEdgeIDs );
+	UFUNCTION( BlueprintCallable, Category="Editable Mesh" ) void CreateMissingPolygonHoleEdges( const FPolygonID PolygonID, const int32 HoleNumber, TArray<FEdgeID>& OutNewEdgeIDs );
 	UFUNCTION( BlueprintCallable, Category="Editable Mesh" ) void SplitEdge( const FEdgeID EdgeID, const TArray<float>& Splits, TArray<FVertexID>& OutNewVertexIDs );
 	UFUNCTION( BlueprintCallable, Category="Editable Mesh" ) void InsertEdgeLoop( const FEdgeID EdgeID, const TArray<float>& Splits, TArray<FEdgeID>& OutNewEdgeIDs );
 	UFUNCTION( BlueprintCallable, Category="Editable Mesh" ) void SplitPolygons( const TArray<FPolygonToSplit>& PolygonsToSplit, TArray<FEdgeID>& OutNewEdgeIDs );
-	UFUNCTION( BlueprintCallable, Category="Editable Mesh" ) void DeleteEdgeAndConnectedPolygons( const FEdgeID EdgeID, const bool bDeleteOrphanedEdges, const bool bDeleteOrphanedVertices, const bool bDeleteEmptySections );
-	UFUNCTION( BlueprintCallable, Category="Editable Mesh" ) void DeleteVertexAndConnectedEdgesAndPolygons( const FVertexID VertexID, const bool bDeleteOrphanedEdges, const bool bDeleteOrphanedVertices, const bool bDeleteEmptySections );
-	UFUNCTION( BlueprintCallable, Category="Editable Mesh" ) void CreateEmptyVertexRange( const int32 NumVerticesToAdd, TArray<FVertexID>& OutNewVertexIDs );
+	UFUNCTION( BlueprintCallable, Category="Editable Mesh" ) void DeleteEdgeAndConnectedPolygons( const FEdgeID EdgeID, const bool bDeleteOrphanedEdges, const bool bDeleteOrphanedVertices, const bool bDeleteOrphanedVertexInstances, const bool bDeleteEmptyPolygonGroups );
+	UFUNCTION( BlueprintCallable, Category="Editable Mesh" ) void DeleteVertexAndConnectedEdgesAndPolygons( const FVertexID VertexID, const bool bDeleteOrphanedEdges, const bool bDeleteOrphanedVertices, const bool bDeleteOrphanedVertexInstances, const bool bDeleteEmptyPolygonGroups );
 	UFUNCTION( BlueprintCallable, Category="Editable Mesh" ) void DeleteOrphanVertices( const TArray<FVertexID>& VertexIDsToDelete );
+	UFUNCTION( BlueprintCallable, Category="Editable Mesh" ) void DeleteVertexInstances( const TArray<FVertexInstanceID>& VertexInstanceIDsToDelete, const bool bDeleteOrphanedVertices );
 	UFUNCTION( BlueprintCallable, Category="Editable Mesh" ) void DeleteEdges( const TArray<FEdgeID>& EdgeIDsToDelete, const bool bDeleteOrphanedVertices );
+	UFUNCTION( BlueprintCallable, Category="Editable Mesh" ) void CreateEmptyVertexRange( const int32 NumVerticesToCreate, TArray<FVertexID>& OutNewVertexIDs );
 	UFUNCTION( BlueprintCallable, Category="Editable Mesh" ) void CreateVertices( const TArray<FVertexToCreate>& VerticesToCreate, TArray<FVertexID>& OutNewVertexIDs );
+	UFUNCTION( BlueprintCallable, Category="Editable Mesh" ) void CreateVertexInstances( const TArray<FVertexInstanceToCreate>& VertexInstancesToCreate, TArray<FVertexInstanceID>& OutNewVertexInstanceIDs );
 	UFUNCTION( BlueprintCallable, Category="Editable Mesh" ) void CreateEdges( const TArray<FEdgeToCreate>& EdgesToCreate, TArray<FEdgeID>& OutNewEdgeIDs );
-	UFUNCTION( BlueprintCallable, Category="Editable Mesh" ) void CreatePolygons( const TArray<FPolygonToCreate>& PolygonsToCreate, TArray<FPolygonRef>& OutNewPolygonRefs, TArray<FEdgeID>& OutNewEdgeIDs );
-	UFUNCTION( BlueprintCallable, Category="Editable Mesh" ) void DeletePolygons( const TArray<FPolygonRef>& PolygonRefsToDelete, const bool bDeleteOrphanedEdges, const bool bDeleteOrphanedVertices, const bool bDeleteEmptySections );
+	UFUNCTION( BlueprintCallable, Category="Editable Mesh" ) void CreatePolygons( const TArray<FPolygonToCreate>& PolygonsToCreate, TArray<FPolygonID>& OutNewPolygonIDs, TArray<FEdgeID>& OutNewEdgeIDs );
+	UFUNCTION( BlueprintCallable, Category="Editable Mesh" ) void DeletePolygons( const TArray<FPolygonID>& PolygonIDsToDelete, const bool bDeleteOrphanedEdges, const bool bDeleteOrphanedVertices, const bool bDeleteOrphanedVertexInstances, const bool bDeleteEmptyPolygonGroups );
 	UFUNCTION( BlueprintCallable, Category="Editable Mesh" ) void SetVerticesAttributes( const TArray<FAttributesForVertex>& AttributesForVertices );
+	UFUNCTION( BlueprintCallable, Category="Editable Mesh" ) void SetVertexInstancesAttributes( const TArray<FAttributesForVertexInstance>& AttributesForVertexInstances );
 	UFUNCTION( BlueprintCallable, Category="Editable Mesh" ) void SetEdgesAttributes( const TArray<FAttributesForEdge>& AttributesForEdges );
 	UFUNCTION( BlueprintCallable, Category="Editable Mesh" ) void SetPolygonsVertexAttributes( const TArray<FVertexAttributesForPolygon>& VertexAttributesForPolygons );
-	UFUNCTION( BlueprintCallable, Category="Editable Mesh" ) void TryToRemovePolygonEdge( const FEdgeID EdgeID, bool& bOutWasEdgeRemoved, FPolygonRef& OutNewPolygonRef );
+	UFUNCTION( BlueprintCallable, Category="Editable Mesh" ) void ChangePolygonsVertexInstances( const TArray<FChangeVertexInstancesForPolygon>& VertexInstancesForPolygons );
+	UFUNCTION( BlueprintCallable, Category="Editable Mesh" ) void TryToRemovePolygonEdge( const FEdgeID EdgeID, bool& bOutWasEdgeRemoved, FPolygonID& OutNewPolygonID );
 	UFUNCTION( BlueprintCallable, Category="Editable Mesh" ) void TryToRemoveVertex( const FVertexID VertexID, bool& bOutWasVertexRemoved, FEdgeID& OutNewEdgeID );
-	UFUNCTION( BlueprintCallable, Category="Editable Mesh" ) void ExtrudePolygons( const TArray<FPolygonRef>& Polygons, const float ExtrudeDistance, const bool bKeepNeighborsTogether, TArray<FPolygonRef>& OutNewExtrudedFrontPolygons );
+	UFUNCTION( BlueprintCallable, Category="Editable Mesh" ) void ExtrudePolygons( const TArray<FPolygonID>& Polygons, const float ExtrudeDistance, const bool bKeepNeighborsTogether, TArray<FPolygonID>& OutNewExtrudedFrontPolygons );
 	UFUNCTION( BlueprintCallable, Category="Editable Mesh" ) void ExtendEdges( const TArray<FEdgeID>& EdgeIDs, const bool bWeldNeighbors, TArray<FEdgeID>& OutNewExtendedEdgeIDs );
 	UFUNCTION( BlueprintCallable, Category="Editable Mesh" ) void ExtendVertices( const TArray<FVertexID>& VertexIDs, const bool bOnlyExtendClosestEdge, const FVector ReferencePosition, TArray<FVertexID>& OutNewExtendedVertexIDs );
-	UFUNCTION( BlueprintCallable, Category="Editable Mesh" ) void InsetPolygons( const TArray<FPolygonRef>& PolygonRefs, const float InsetFixedDistance, const float InsetProgressTowardCenter, const EInsetPolygonsMode Mode, TArray<FPolygonRef>& OutNewCenterPolygonRefs, TArray<FPolygonRef>& OutNewSidePolygonRefs );
-	UFUNCTION( BlueprintCallable, Category="Editable Mesh" ) void BevelPolygons( const TArray<FPolygonRef>& PolygonRefs, const float BevelFixedDistance, const float BevelProgressTowardCenter, TArray<FPolygonRef>& OutNewCenterPolygonRefs, TArray<FPolygonRef>& OutNewSidePolygonRefs );
-	UFUNCTION( BlueprintCallable, Category="Editable Mesh" ) void GenerateNormalsAndTangentsForPolygons( const TArray< FPolygonRef >& PolygonRefs );
-	UFUNCTION( BlueprintCallable, Category="Editable Mesh" ) void GenerateNormalsAndTangentsForPolygonsAndAdjacents( const TArray< FPolygonRef >& PolygonRefs );
+	UFUNCTION( BlueprintCallable, Category="Editable Mesh" ) void InsetPolygons( const TArray<FPolygonID>& PolygonIDs, const float InsetFixedDistance, const float InsetProgressTowardCenter, const EInsetPolygonsMode Mode, TArray<FPolygonID>& OutNewCenterPolygonIDs, TArray<FPolygonID>& OutNewSidePolygonIDs );
+	UFUNCTION( BlueprintCallable, Category="Editable Mesh" ) void BevelPolygons( const TArray<FPolygonID>& PolygonIDs, const float BevelFixedDistance, const float BevelProgressTowardCenter, TArray<FPolygonID>& OutNewCenterPolygonIDs, TArray<FPolygonID>& OutNewSidePolygonIDs );
 	UFUNCTION( BlueprintCallable, Category="Editable Mesh" ) void SetVerticesCornerSharpness( const TArray<FVertexID>& VertexIDs, const TArray<float>& VerticesNewCornerSharpness );
 	UFUNCTION( BlueprintCallable, Category="Editable Mesh" ) void SetEdgesCreaseSharpness( const TArray<FEdgeID>& EdgeIDs, const TArray<float>& EdgesNewCreaseSharpness );
 	UFUNCTION( BlueprintCallable, Category="Editable Mesh" ) void SetEdgesHardness( const TArray<FEdgeID>& EdgeIDs, const TArray<bool>& EdgesNewIsHard );
 	UFUNCTION( BlueprintCallable, Category="Editable Mesh" ) void SetEdgesHardnessAutomatically( const TArray<FEdgeID>& EdgeIDs, const float MaxDotProductForSoftEdge );	// @todo mesheditor: Not used for anything yet.  Remove it?  Use it during import/convert?
 	UFUNCTION( BlueprintCallable, Category="Editable Mesh" ) void SetEdgesVertices( const TArray<FVerticesForEdge>& VerticesForEdges );
-	UFUNCTION( BlueprintCallable, Category="Editable Mesh" ) void InsertPolygonPerimeterVertices( const FPolygonRef PolygonRef, const int32 InsertBeforeVertexNumber, const TArray<FVertexAndAttributes>& VerticesToInsert );
-	UFUNCTION( BlueprintCallable, Category="Editable Mesh" ) void RemovePolygonPerimeterVertices( const FPolygonRef PolygonRef, const int32 FirstVertexNumberToRemove, const int32 NumVerticesToRemove );
-	UFUNCTION( BlueprintCallable, Category="Editable Mesh" ) void FlipPolygons( const TArray<FPolygonRef>& PolygonRefs );
-	UFUNCTION( BlueprintCallable, Category="Editable Mesh" ) void TriangulatePolygons( const TArray<FPolygonRef>& PolygonRefs, TArray<FPolygonRef>& OutNewTrianglePolygons );
-	UFUNCTION( BlueprintCallable, Category="Editable Mesh" ) FSectionID CreateSection( const FSectionToCreate& SectionToCreate );
-	UFUNCTION( BlueprintCallable, Category="Editable Mesh" ) void DeleteSection( const FSectionID SectionID );
-	UFUNCTION( BlueprintCallable, Category="Editable Mesh" ) void AssignMaterialToPolygons( const TArray<FPolygonRef>& PolygonRefs, UMaterialInterface* Material, TArray<FPolygonRef>& NewPolygonRefs );
+	UFUNCTION( BlueprintCallable, Category="Editable Mesh" ) void InsertPolygonPerimeterVertices( const FPolygonID PolygonID, const int32 InsertBeforeVertexNumber, const TArray<FVertexAndAttributes>& VerticesToInsert );
+	UFUNCTION( BlueprintCallable, Category="Editable Mesh" ) void RemovePolygonPerimeterVertices( const FPolygonID PolygonID, const int32 FirstVertexNumberToRemove, const int32 NumVerticesToRemove, const bool bDeleteOrphanedVertexInstances );
+	UFUNCTION( BlueprintCallable, Category="Editable Mesh" ) void FlipPolygons( const TArray<FPolygonID>& PolygonIDs );
+	UFUNCTION( BlueprintCallable, Category="Editable Mesh" ) void TriangulatePolygons( const TArray<FPolygonID>& PolygonIDs, TArray<FPolygonID>& OutNewTrianglePolygons );
+	UFUNCTION( BlueprintCallable, Category="Editable Mesh" ) void CreatePolygonGroups( const TArray<FPolygonGroupToCreate>& PolygonGroupsToCreate, TArray<FPolygonGroupID>& OutNewPolygonGroupIDs );
+	UFUNCTION( BlueprintCallable, Category="Editable Mesh" ) void DeletePolygonGroups( const TArray<FPolygonGroupID>& PolygonGroupIDs );
+	UFUNCTION( BlueprintCallable, Category="Editable Mesh" ) void AssignMaterialToPolygons( const TArray<FPolygonID>& PolygonIDs, UMaterialInterface* Material, TArray<FPolygonID>& OutNewPolygonIDs );
 	UFUNCTION( BlueprintCallable, Category="Editable Mesh" ) void WeldVertices( const TArray<FVertexID>& VertexIDs, FVertexID& OutNewVertexID );
-	UFUNCTION( BlueprintCallable, Category="Editable Mesh" ) void TessellatePolygons( const TArray<FPolygonRef>& PolygonRefs, const ETriangleTessellationMode TriangleTessellationMode, TArray<FPolygonRef>& OutNewPolygonRefs );
-	UFUNCTION( BlueprintCallable, Category="Editable Mesh" ) void SetTextureCoordinateCount(int32 NumTexCoords);
-	UFUNCTION( BlueprintCallable, Category="Editable Mesh" ) void QuadrangulateMesh( TArray<FPolygonRef>& NewPolygonRefs );
+	UFUNCTION( BlueprintCallable, Category="Editable Mesh" ) void TessellatePolygons( const TArray<FPolygonID>& PolygonIDs, const ETriangleTessellationMode TriangleTessellationMode, TArray<FPolygonID>& OutNewPolygonIDs );
+	UFUNCTION( BlueprintCallable, Category="Editable Mesh" ) void SetTextureCoordinateCount( const int32 NumTexCoords );
+	UFUNCTION( BlueprintCallable, Category="Editable Mesh" ) void QuadrangulateMesh( TArray<FPolygonID>& OutNewPolygonIDs );
+	UFUNCTION( BlueprintCallable, Category="Editable Mesh" ) void GeneratePolygonTangentsAndNormals( const TArray<FPolygonID>& PolygonIDs );
+
 protected:
 
 	/** Adds a new change that can be used to undo a modification that happened to the mesh.  If bAllowUndo is false, then
 	    the undo will not be stored */
 	void AddUndo( TUniquePtr<FChange> NewUndo );
 
+public:
+	// @todo mesheditor: temporarily changed access to public so the adapter can call it when building the editable mesh from the source static mesh. Think about this some more.
 	/** Refreshes the entire OpenSubdiv state for this mesh and generates subdivision geometry (if the mesh is configured to have subdivision levels) */
 	void RefreshOpenSubdiv();
 
+protected:
+
 	/** Generates new limit surface data and caches it (if the mesh is configured to have subdivision levels).  This can be used to update the
 	    limit surface after a non-topology-effective edit has happened */
 	void GenerateOpenSubdivLimitSurfaceData();
@@ -331,12 +754,37 @@ protected:
 
 private:
 
-	void BevelOrInsetPolygons( const TArray<FPolygonRef>& PolygonRefs, const float InsetFixedDistance, const float InsetProgressTowardCenter, const EInsetPolygonsMode Mode, const bool bShouldBevel, TArray<FPolygonRef>& OutNewCenterPolygonRefs, TArray<FPolygonRef>& OutNewSidePolygonRefs );
+	void BevelOrInsetPolygons( const TArray<FPolygonID>& PolygonIDs, const float InsetFixedDistance, const float InsetProgressTowardCenter, const EInsetPolygonsMode Mode, const bool bShouldBevel, TArray<FPolygonID>& OutNewCenterPolygonIDs, TArray<FPolygonID>& OutNewSidePolygonIDs );
 
-	void GenerateTangentsForPolygons( const TArray< FPolygonRef >& PolygonRefs );
+	/** Method used internally by the quadrangulate action */
+	void QuadrangulatePolygonGroup( const FPolygonGroupID PolygonGroupID, TArray<FPolygonID>& OutNewPolygonIDs );
 
+	/** Called during end modification to generate tangents and normals on the pending polygon list */
+	void GenerateTangentsAndNormals();
 
-protected:
+	/** Called during end modification to retriangulate polygons in the pending polygon list */
+	void RetriangulatePolygons();
+
+public:
+
+	/** Each editable vertex in this mesh. */
+	TSparseArray<FMeshVertex> Vertices;
+
+	/** Sparse array of rendering vertices, that matches the vertices in the mesh vertex buffers */
+	TSparseArray<FMeshVertexInstance> VertexInstances;
+
+	/** All editable mesh edges.  Note that some of these edges will be internal polygon edges, synthesized while
+	triangulating polygons into triangles.  Static meshes currently only support triangles. */
+	TSparseArray<FMeshEdge> Edges;
+
+	/** All of the polygons in this mesh */
+	TSparseArray<FMeshPolygon> Polygons;
+
+	/** All of the polygon groups in this mesh */
+	TSparseArray<FMeshPolygonGroup> PolygonGroups;
+
+// @todo mesheditor: sort out member access. Currently StaticMesh adapter relies on accessing this stuff directly
+//protected:
 
 	/** The sub-mesh we came from */
 	FEditableMeshSubMeshAddress SubMeshAddress;
@@ -348,6 +796,10 @@ protected:
 	    mesh since the last time that MakeUndo() was called. */
 	TUniquePtr<FCompoundChangeInput> Undo;
 
+	/** Adapters registered with this editable mesh */
+	UPROPERTY()
+	TArray<UEditableMeshAdapter*> Adapters;
+
 	/** The number of texture coordinates stored on the vertices of this mesh */
 	UPROPERTY( BlueprintReadOnly, Category="Editable Mesh" )
 	int32 TextureCoordinateCount;
@@ -356,6 +808,12 @@ protected:
 	UPROPERTY( BlueprintReadOnly, Category="Editable Mesh" ) 
 	int32 SubdivisionCount;
 
+	/** List of polygons which need their tangent basis recalculating (and consequently their associated vertex instances) */
+	TSet<FPolygonID> PolygonsPendingNewTangentBasis;
+
+	/** List of polygons requiring retriangulation */
+	TSet<FPolygonID> PolygonsPendingTriangulation;
+
 	/** True if StartModification() has been called.  Call EndModification() when you've finished changing the mesh. */
 	bool bIsBeingModified;
 
@@ -365,6 +823,13 @@ protected:
 	/** While the mesh is being edited (between calls to StartModification() and EndModification()), stores whether topology could be affected */
 	EMeshTopologyChange CurrentToplogyChange;
 
+	/** Counter to determine when we should compact data */
+	UPROPERTY()
+	int32 PendingCompactCounter;
+
+	/** Data will be compacted after this many topology modifying actions. */
+	static const int32 CompactFrequency = 10;
+
 	/** OpenSubdiv topology refiner object.  This is generated for meshes that have subdivision levels, and reused to generate new limit surfaces 
 	    when geometry is moved.  When the mesh's topology changes, this object is regenerated from scratch. */
 	TSharedPtr<OpenSubdiv::v3_2_0::Far::TopologyRefiner> OsdTopologyRefiner;
diff --git a/Engine/Source/Runtime/MeshEditingRuntime/Public/EditableMeshAdapter.h b/Engine/Source/Runtime/MeshEditingRuntime/Public/EditableMeshAdapter.h
new file mode 100644
--- /dev/null
+++ b/Engine/Source/Runtime/MeshEditingRuntime/Public/EditableMeshAdapter.h
@@ -0,0 +1,49 @@
+// Copyright 1998-2015 Epic Games, Inc. All Rights Reserved.
+
+#pragma once
+
+#include "CoreMinimal.h"
+#include "UObject/Object.h"
+#include "EditableMesh.h"
+#include "EditableMeshTypes.h"
+#include "EditableMeshAdapter.generated.h"
+
+
+UCLASS(Abstract)
+class UEditableMeshAdapter : public UObject
+{
+	GENERATED_BODY()
+
+public:
+
+	virtual void OnRebuildRenderMeshStart( const UEditableMesh* EditableMesh, const bool bRefreshBounds, const bool bInvalidateLighting ) PURE_VIRTUAL(,);
+	virtual void OnRebuildRenderMesh( const UEditableMesh* EditableMesh ) PURE_VIRTUAL(,);
+	virtual void OnRebuildRenderMeshFinish( const UEditableMesh* EditableMesh, const bool bUpdateCollision ) PURE_VIRTUAL(,);
+	virtual void OnStartModification( const UEditableMesh* EditableMesh, const EMeshModificationType MeshModificationType, const EMeshTopologyChange MeshTopologyChange ) PURE_VIRTUAL(,);
+	virtual void OnEndModification( const UEditableMesh* EditableMesh ) PURE_VIRTUAL(,);
+	virtual void OnReindexElements( const UEditableMesh* EditableMesh, const FElementIDRemappings& Remappings ) PURE_VIRTUAL(,);
+	virtual bool IsCommitted( const UEditableMesh* EditableMesh ) const PURE_VIRTUAL(,return false;);
+	virtual bool IsCommittedAsInstance( const UEditableMesh* EditableMesh ) const PURE_VIRTUAL(,return false;);
+	virtual void OnCommit( UEditableMesh* EditableMesh ) PURE_VIRTUAL(,);
+	virtual UEditableMesh* OnCommitInstance( UEditableMesh* EditableMesh, UPrimitiveComponent* ComponentToInstanceTo ) PURE_VIRTUAL(,return nullptr;);
+	virtual void OnRevert( UEditableMesh* EditableMesh ) PURE_VIRTUAL(,);
+	virtual UEditableMesh* OnRevertInstance( UEditableMesh* EditableMesh ) PURE_VIRTUAL(,return nullptr;);
+	virtual void OnPropagateInstanceChanges( UEditableMesh* EditableMesh ) PURE_VIRTUAL(,);
+
+	virtual void OnDeleteVertexInstances( const UEditableMesh* EditableMesh, const TArray<FVertexInstanceID>& VertexInstanceIDs ) PURE_VIRTUAL(,);
+	virtual void OnDeleteOrphanVertices( const UEditableMesh* EditableMesh, const TArray<FVertexID>& VertexIDs ) PURE_VIRTUAL(,);
+	virtual void OnCreateEmptyVertexRange( const UEditableMesh* EditableMesh, const TArray<FVertexID>& VertexIDs ) PURE_VIRTUAL(,);
+	virtual void OnCreateVertices( const UEditableMesh* EditableMesh, const TArray<FVertexID>& VertexIDs ) PURE_VIRTUAL(,);
+	virtual void OnCreateVertexInstances( const UEditableMesh* EditableMesh, const TArray<FVertexInstanceID>& VertexInstanceIDs ) PURE_VIRTUAL(,);
+	virtual void OnSetVertexAttribute( const UEditableMesh* EditableMesh, const FVertexID VertexID, const FName AttributeName, const int32 AttributeIndex, const FVector4 AttributeValue ) PURE_VIRTUAL(,);
+	virtual void OnSetVertexInstanceAttribute( const UEditableMesh* EditableMesh, const FVertexInstanceID VertexInstanceID, const FName AttributeName, const int32 AttributeIndex, const FVector4 AttributeValue ) PURE_VIRTUAL(,);
+	virtual void OnCreateEdges( const UEditableMesh* EditableMesh, const TArray<FEdgeID>& EdgeIDs ) PURE_VIRTUAL(,);
+	virtual void OnDeleteEdges( const UEditableMesh* EditableMesh, const TArray<FEdgeID>& EdgeIDs ) PURE_VIRTUAL(,);
+	virtual void OnSetEdgeAttribute( const UEditableMesh* EditableMesh, const FEdgeID EdgeID, const FName AttributeName, const int32 AttributeIndex, const FVector4 AttributeValue ) PURE_VIRTUAL(,);
+	virtual void OnCreatePolygons( const UEditableMesh* EditableMesh, const TArray<FPolygonID>& PolygonIDs ) PURE_VIRTUAL(,);
+	virtual void OnDeletePolygons( const UEditableMesh* EditableMesh, const TArray<FPolygonID>& PolygonIDs ) PURE_VIRTUAL(,);
+	virtual void OnChangePolygonVertexInstances( const UEditableMesh* EditableMesh, const TArray<FPolygonID>& PolygonIDs ) PURE_VIRTUAL(,);
+	virtual void OnCreatePolygonGroups( const UEditableMesh* EditableMesh, const TArray<FPolygonGroupID>& PolygonGroupIDs ) PURE_VIRTUAL(,);
+	virtual void OnDeletePolygonGroups( const UEditableMesh* EditableMesh, const TArray<FPolygonGroupID>& PolygonGroupIDs ) PURE_VIRTUAL(,);
+	virtual void OnRetriangulatePolygons( const UEditableMesh* EditableMesh, const TArray<FPolygonID>& PolygonIDs ) PURE_VIRTUAL(,);
+};
diff --git a/Engine/Source/Runtime/MeshEditingRuntime/Public/EditableMeshTypes.h b/Engine/Source/Runtime/MeshEditingRuntime/Public/EditableMeshTypes.h
--- a/Engine/Source/Runtime/MeshEditingRuntime/Public/EditableMeshTypes.h
+++ b/Engine/Source/Runtime/MeshEditingRuntime/Public/EditableMeshTypes.h
@@ -142,7 +142,7 @@ struct FElementID	// @todo mesheditor script: BP doesn't have name spaces, so we
 
 	FString ToString() const
 	{
-		return FString::Printf( TEXT( "%lu" ), GetValue() );
+		return ( IDValue == Invalid.GetValue() ) ? TEXT( "Invalid" ) : FString::Printf( TEXT( "%d" ), IDValue );
 	}
 
 	friend FArchive& operator<<( FArchive& Ar, FElementID& Element )
@@ -192,118 +192,118 @@ struct FVertexID : public FElementID
 
 
 USTRUCT( BlueprintType )
-struct FEdgeID : public FElementID
+struct FVertexInstanceID : public FElementID
 {
 	GENERATED_BODY()
 
-	FEdgeID()
+	FVertexInstanceID()
 	{
 	}
 
-	explicit FEdgeID( const FElementID InitElementID )
+	explicit FVertexInstanceID( const FElementID InitElementID )
 		: FElementID( InitElementID.GetValue() )
 	{
 	}
 
-	explicit FEdgeID( const int32 InitIDValue )
+	explicit FVertexInstanceID( const uint32 InitIDValue )
 		: FElementID( InitIDValue )
 	{
 	}
 
-	FORCEINLINE friend uint32 GetTypeHash( const FEdgeID& Other )
+	FORCEINLINE friend uint32 GetTypeHash( const FVertexInstanceID& Other )
 	{
 		return GetTypeHash( Other.IDValue );
 	}
 
-	/** Invalid edge ID */
-	MESHEDITINGRUNTIME_API static const FEdgeID Invalid;
+	/** Invalid rendering vertex ID */
+	MESHEDITINGRUNTIME_API static const FVertexInstanceID Invalid;
 };
 
 
 USTRUCT( BlueprintType )
-struct FSectionID : public FElementID
+struct FEdgeID : public FElementID
 {
 	GENERATED_BODY()
 
-	FSectionID()
+	FEdgeID()
 	{
 	}
 
-	explicit FSectionID( const FElementID InitElementID )
+	explicit FEdgeID( const FElementID InitElementID )
 		: FElementID( InitElementID.GetValue() )
 	{
 	}
 
-	explicit FSectionID( const int32 InitIDValue )
+	explicit FEdgeID( const int32 InitIDValue )
 		: FElementID( InitIDValue )
 	{
 	}
 
-	FORCEINLINE friend uint32 GetTypeHash( const FSectionID& Other )
+	FORCEINLINE friend uint32 GetTypeHash( const FEdgeID& Other )
 	{
 		return GetTypeHash( Other.IDValue );
 	}
 
-	/** Invalid section ID */
-	MESHEDITINGRUNTIME_API static const FSectionID Invalid;
+	/** Invalid edge ID */
+	MESHEDITINGRUNTIME_API static const FEdgeID Invalid;
 };
 
 
 USTRUCT( BlueprintType )
-struct FPolygonID : public FElementID
+struct FPolygonGroupID : public FElementID
 {
 	GENERATED_BODY()
 
-	FPolygonID()
+	FPolygonGroupID()
 	{
 	}
 
-	explicit FPolygonID( const FElementID InitElementID )
+	explicit FPolygonGroupID( const FElementID InitElementID )
 		: FElementID( InitElementID.GetValue() )
 	{
 	}
 
-	explicit FPolygonID( const int32 InitIDValue )
+	explicit FPolygonGroupID( const int32 InitIDValue )
 		: FElementID( InitIDValue )
 	{
 	}
 
-	FORCEINLINE friend uint32 GetTypeHash( const FPolygonID& Other )
+	FORCEINLINE friend uint32 GetTypeHash( const FPolygonGroupID& Other )
 	{
 		return GetTypeHash( Other.IDValue );
 	}
 
-	/** Invalid polygon ID */
-	MESHEDITINGRUNTIME_API static const FPolygonID Invalid;	// @todo mesheditor script: Can we expose these to BP nicely?	Do we even need to?
+	/** Invalid section ID */
+	MESHEDITINGRUNTIME_API static const FPolygonGroupID Invalid;
 };
 
 
-USTRUCT()
-struct FTriangleID : public FElementID
+USTRUCT( BlueprintType )
+struct FPolygonID : public FElementID
 {
 	GENERATED_BODY()
 
-	FTriangleID()
+	FPolygonID()
 	{
 	}
 
-	explicit FTriangleID( const FElementID InitElementID )
+	explicit FPolygonID( const FElementID InitElementID )
 		: FElementID( InitElementID.GetValue() )
 	{
 	}
 
-	explicit FTriangleID( const uint32 InitIDValue )
+	explicit FPolygonID( const int32 InitIDValue )
 		: FElementID( InitIDValue )
 	{
 	}
 
-	FORCEINLINE friend uint32 GetTypeHash( const FTriangleID& Other )
+	FORCEINLINE friend uint32 GetTypeHash( const FPolygonID& Other )
 	{
 		return GetTypeHash( Other.IDValue );
 	}
 
-	/** Invalid triangle ID */
-	MESHEDITINGRUNTIME_API static const FTriangleID Invalid;
+	/** Invalid polygon ID */
+	MESHEDITINGRUNTIME_API static const FPolygonID Invalid;	// @todo mesheditor script: Can we expose these to BP nicely?	Do we even need to?
 };
 
 
@@ -365,61 +365,6 @@ struct FEditableMeshSubMeshAddress
 };
 
 
-
-USTRUCT( BlueprintType )
-struct FPolygonRef
-{
-	GENERATED_BODY()
-
-	UPROPERTY( BlueprintReadWrite, Category="Editable Mesh" )
-	FSectionID SectionID;
-
-	UPROPERTY( BlueprintReadWrite, Category="Editable Mesh" )
-	FPolygonID PolygonID;
-
-	FPolygonRef()
-	{
-	}
-
-	FPolygonRef( const FSectionID InitSectionID, const FPolygonID InitPolygonID )
-		: SectionID( InitSectionID ),
-		  PolygonID( InitPolygonID )
-	{
-	}
-
-	inline bool operator==( const FPolygonRef& Other ) const
-	{
-		return SectionID == Other.SectionID && PolygonID == Other.PolygonID;
-	}
-
-	inline bool operator!=( const FPolygonRef& Other ) const
-	{
-		return SectionID != Other.SectionID || PolygonID != Other.PolygonID;
-	}
-
-	/** Hashing */
-	FORCEINLINE friend uint32 GetTypeHash( const FPolygonRef& Other )
-	{
-		return GetTypeHash( ( uint64( Other.PolygonID.GetValue() ) << 32 ) | uint64( Other.SectionID.GetValue() ) );
-	}
-
-	FString ToString() const
-	{
-		return FString::Printf( TEXT( "SectionID:%s, PolygonID:%s" ), *SectionID.ToString(), *PolygonID.ToString() );
-	}
-
-	friend FArchive& operator<<( FArchive& Ar, FPolygonRef& PolygonRef )
-	{
-		Ar << PolygonRef.SectionID;
-		Ar << PolygonRef.PolygonID;
-		return Ar;
-	}
-
-	/** Invalid polygon ref */
-	MESHEDITINGRUNTIME_API static const FPolygonRef Invalid;
-};
-
-
 UCLASS( abstract )
 class MESHEDITINGRUNTIME_API UEditableMeshAttribute : public UObject
 {
@@ -444,7 +389,7 @@ public:
 	}
 
 	//
-	// Polygon-specific vertex data (can also be set on the vertex itself to update all polygon vertices.)
+	// Vertex instance data
 	//
 
 	/** Static: The attribute name for vertex normal (tangent Z) */
@@ -493,6 +438,22 @@ public:
 		return EdgeCreaseSharpnessName;
 	}
 
+	//
+	// Polygons
+	//
+
+	/** Static: The attribute name for polygon normal */
+	UFUNCTION( BlueprintPure, Category="Editable Mesh" ) static inline FName PolygonNormal()
+	{
+		return PolygonNormalName;
+	}
+
+	/** Static: The attribute name for polygon center */
+	UFUNCTION( BlueprintPure, Category="Editable Mesh" ) static inline FName PolygonCenter()
+	{
+		return PolygonCenterName;
+	}
+
 private:
 
 	static const FName VertexPositionName;
@@ -504,6 +465,8 @@ private:
 	static const FName VertexColorName;
 	static const FName EdgeIsHardName;
 	static const FName EdgeCreaseSharpnessName;
+	static const FName PolygonNormalName;
+	static const FName PolygonCenterName;
 };
 
 
@@ -757,6 +720,41 @@ struct FVertexToCreate
 };
 
 
+USTRUCT( BlueprintType )
+struct FVertexInstanceToCreate
+{
+	GENERATED_BODY()
+
+	/** Vertex ID which is being instanced */
+	UPROPERTY( BlueprintReadWrite, Category="Editable Mesh" )
+	FVertexID VertexID;
+
+	/** Attributes of this vertex instance */
+	UPROPERTY( BlueprintReadWrite, Category="Editable Mesh" )
+	FMeshElementAttributeList VertexInstanceAttributes;
+
+	/** The original ID of the vertex instance.  Should only be used by the undo system. */
+	UPROPERTY()
+	FVertexInstanceID OriginalVertexInstanceID;
+
+	FVertexInstanceToCreate()
+		: VertexID( FVertexID::Invalid ),
+		  VertexInstanceAttributes(),
+		  OriginalVertexInstanceID( FVertexInstanceID::Invalid )
+	{
+	}
+
+	FString ToString() const
+	{
+		return FString::Printf(
+			TEXT( "VertexID:%s, VertexInstanceAttributes:%s, OriginalVertexInstanceID:%s" ),
+			*VertexID.ToString(),
+			*VertexInstanceAttributes.ToString(),
+			*OriginalVertexInstanceID.ToString() );
+	}
+};
+
+
 USTRUCT( BlueprintType )
 struct FEdgeToCreate
 {
@@ -772,7 +770,7 @@ struct FEdgeToCreate
 
 	/** The polygons that are connected to this edge. */
 	UPROPERTY( BlueprintReadWrite, Category="Editable Mesh" )
-	TArray<FPolygonRef> ConnectedPolygons;
+	TArray<FPolygonID> ConnectedPolygons;
 
 	/** Attributes of this edge itself */
 	UPROPERTY( BlueprintReadWrite, Category="Editable Mesh" )
@@ -811,17 +809,22 @@ struct FVertexAndAttributes
 {
 	GENERATED_BODY()
 
-	/** The vertex ID to insert into the polygon */
+	/** The vertex instance ID to insert into the polygon, or FVertexInstanceID::Invalid to create a new vertex instance with the given attributes */
+	UPROPERTY( BlueprintReadWrite, Category="Editable Mesh" )
+	FVertexInstanceID VertexInstanceID;
+
+	/** The vertex ID to insert into the polygon, if no valid vertex instance ID was supplied. */
 	UPROPERTY( BlueprintReadWrite, Category="Editable Mesh" )
 	FVertexID VertexID;
 
-	/** A list of polygon attributes to set for the vertex on the polygon we're inserting it into */
+	/** A list of polygon attributes to set for the vertex on the polygon we're inserting it into, if no valid vertex instance ID was supplied. */
 	UPROPERTY( BlueprintReadWrite, Category="Editable Mesh" )
 	FMeshElementAttributeList PolygonVertexAttributes;
 
 	/** Default constructor */
 	FVertexAndAttributes()
-		: VertexID( 0 ),
+		: VertexInstanceID( FVertexInstanceID::Invalid ),
+		  VertexID( FVertexID::Invalid ),
 		  PolygonVertexAttributes()
 	{
 	}
@@ -829,7 +832,8 @@ struct FVertexAndAttributes
 	FString ToString() const
 	{
 		return FString::Printf(
-			TEXT( "VertexID:%s, PolygonVertexAttributes:%s" ),
+			TEXT( "VertexInstanceID:%s, VertexID:%s, PolygonVertexAttributes:%s" ),
+			*VertexInstanceID.ToString(),
 			*VertexID.ToString(),
 			*PolygonVertexAttributes.ToString() );
 	}
@@ -860,15 +864,31 @@ struct FPolygonHoleVertices
 };
 
 
+UENUM( BlueprintType )
+enum class EPolygonEdgeHardness : uint8
+{
+	/** Create any new edges required by the polygon as soft, leaving existing edges as they are */
+	NewEdgesSoft,
+
+	/** Create any new edges required by the polygon as hard, leaving existing edges as they are */
+	NewEdgesHard,
+
+	/** Set all edges as soft, overriding edge hardness on existing edges */
+	AllEdgesSoft,
+
+	/** Set all edges as hard, overriding edge hardness on existing edges */
+	AllEdgesHard
+};
+
 
 USTRUCT( BlueprintType )
 struct FPolygonToCreate
 {
 	GENERATED_BODY()
 
-	/** The section the polygon will be added to */
+	/** The group the polygon will be added to */
 	UPROPERTY( BlueprintReadWrite, Category="Editable Mesh" )
-	FSectionID SectionID;
+	FPolygonGroupID PolygonGroupID;
 
 	/** Ordered list of vertices that defines the polygon's perimeter, along with the polygon vertex attributes to set for each vertex */
 	UPROPERTY( BlueprintReadWrite, Category="Editable Mesh" )
@@ -882,20 +902,25 @@ struct FPolygonToCreate
 	UPROPERTY( BlueprintReadWrite, Category="Editable Mesh" )
 	FPolygonID OriginalPolygonID;
 
+	/** Whether to create a hard-edged polygon */
+	UPROPERTY( BlueprintReadWrite, Category="Editable Mesh" )
+	EPolygonEdgeHardness PolygonEdgeHardness;
+
 	/** Default constructor */
 	FPolygonToCreate()
-		: SectionID( 0 ),
+		: PolygonGroupID( FPolygonGroupID::Invalid ),
 		  PerimeterVertices(),
 		  PolygonHoles(),
-		  OriginalPolygonID( FPolygonID::Invalid )
+		  OriginalPolygonID( FPolygonID::Invalid ),
+		  PolygonEdgeHardness( EPolygonEdgeHardness::NewEdgesSoft )
 	{
 	}
 
 	FString ToString() const
 	{
 		return FString::Printf(
-			TEXT( "SectionID:%s, PerimeterVertices:%s, PolygonHoles:%s, OriginalPolygonID:%s" ),
-			*SectionID.ToString(),
+			TEXT( "PolygonGroupID:%s, PerimeterVertices:%s, PolygonHoles:%s, OriginalPolygonID:%s" ),
+			*PolygonGroupID.ToString(),
 			*LogHelpers::ArrayToString( PerimeterVertices ),
 			*LogHelpers::ArrayToString( PolygonHoles ),
 			*OriginalPolygonID.ToString() );
@@ -935,7 +960,7 @@ struct FPolygonToSplit
 
 	/** The polygon that we'll be splitting */
 	UPROPERTY( BlueprintReadWrite, Category="Editable Mesh" )
-	FPolygonRef PolygonRef;
+	FPolygonID PolygonID;
 
 	/** A list of pairs of vertices that new edges will be created at.  The pairs must be ordered, and the vertices
 	    must already exist and be connected to the polygon */
@@ -944,7 +969,7 @@ struct FPolygonToSplit
 
 	/** Default constructor */
 	FPolygonToSplit()
-		: PolygonRef( FSectionID::Invalid, FPolygonID::Invalid ),
+		: PolygonID( FPolygonID::Invalid ),
 		  VertexPairsToSplitAt()
 	{
 	}
@@ -983,6 +1008,36 @@ struct FAttributesForVertex
 };
 
 
+USTRUCT( BlueprintType )
+struct FAttributesForVertexInstance
+{
+	GENERATED_BODY()
+
+	/** The vertex instance ID to set attributes on */
+	UPROPERTY( BlueprintReadWrite, Category="Editable Mesh" )
+	FVertexInstanceID VertexInstanceID;
+
+	/** A list of attributes to set for the vertex instance */
+	UPROPERTY( BlueprintReadWrite, Category="Editable Mesh" )
+	FMeshElementAttributeList VertexInstanceAttributes;
+
+	/** Default constructor */
+	FAttributesForVertexInstance()
+		: VertexInstanceID( 0 ),
+		  VertexInstanceAttributes()
+	{
+	}
+
+	FString ToString() const
+	{
+		return FString::Printf(
+			TEXT( "VertexInstanceID:%s, VertexInstanceAttributes:%s" ),
+			*VertexInstanceID.ToString(),
+			*VertexInstanceAttributes.ToString() );
+	}
+};
+
+
 USTRUCT( BlueprintType )
 struct FAttributesForEdge
 {
@@ -996,19 +1051,25 @@ struct FAttributesForEdge
 	UPROPERTY( BlueprintReadWrite, Category="Editable Mesh" )
 	FMeshElementAttributeList EdgeAttributes;
 
+	/** Whether this is being done as part of undo/redo */
+	UPROPERTY( BlueprintReadWrite, Category="Editable Mesh" )
+	bool bIsUndo;
+
 	/** Default constructor */
 	FAttributesForEdge()
 		: EdgeID( 0 ),
-		  EdgeAttributes()
+		  EdgeAttributes(),
+		  bIsUndo( false )
 	{
 	}
 
 	FString ToString() const
 	{
 		return FString::Printf(
-			TEXT( "EdgeID:%s, EdgeAttributes:%s" ),
+			TEXT( "EdgeID:%s, EdgeAttributes:%s, bIsUndo:%s" ),
 			*EdgeID.ToString(),
-			*EdgeAttributes.ToString() );
+			*EdgeAttributes.ToString(),
+			*LogHelpers::BoolToString( bIsUndo ) );
 	}
 };
 
@@ -1045,7 +1106,7 @@ struct FVertexAttributesForPolygon
 
 	/** The polygon to set vertex attributes on */
 	UPROPERTY( BlueprintReadWrite, Category="Editable Mesh" )
-	FPolygonRef PolygonRef;
+	FPolygonID PolygonID;
 
 	/** For each polygon vertex, a list of attributes for that vertex.  Can be left empty if you don't want to set any attributes. */
 	UPROPERTY( BlueprintReadWrite, Category="Editable Mesh" )
@@ -1058,7 +1119,7 @@ struct FVertexAttributesForPolygon
 
 	/** Default constructor */
 	FVertexAttributesForPolygon()
-		: PolygonRef( FSectionID::Invalid, FPolygonID::Invalid ),
+		: PolygonID( FPolygonID::Invalid ),
 		  PerimeterVertexAttributeLists(),
 		  VertexAttributeListsForEachHole()
 	{
@@ -1067,14 +1128,97 @@ struct FVertexAttributesForPolygon
 	FString ToString() const
 	{
 		return FString::Printf(
-			TEXT( "PolygonRef:%s, PerimeterVertexAttributeLists:%s, VertexAttributeListsForEachHole:%s" ),
-			*PolygonRef.ToString(),
+			TEXT( "PolygonID:%s, PerimeterVertexAttributeLists:%s, VertexAttributeListsForEachHole:%s" ),
+			*PolygonID.ToString(),
 			*LogHelpers::ArrayToString( PerimeterVertexAttributeLists ),
 			*LogHelpers::ArrayToString( VertexAttributeListsForEachHole ) );
 	}
 };
 
 
+USTRUCT( BlueprintType )
+struct FVertexIndexAndInstanceID
+{
+	GENERATED_BODY()
+
+	/** Contour index of the vertex to change */
+	UPROPERTY( BlueprintReadWrite, Category="Editable Mesh" )
+	int32 ContourIndex;
+
+	/** New vertex instance ID to assign */
+	UPROPERTY( BlueprintReadWrite, Category="Editable Mesh" )
+	FVertexInstanceID VertexInstanceID;
+
+	FString ToString() const
+	{
+		return FString::Printf(
+			TEXT( "ContourIndex:%d, VertexInstanceIDs:%s" ),
+			ContourIndex,
+			*VertexInstanceID.ToString() );
+	}
+};
+
+
+USTRUCT( BlueprintType )
+struct FVertexInstancesForPolygonHole
+{
+	GENERATED_BODY()
+
+	/** A list of vertex instance IDs for a polygon hole. */
+	UPROPERTY( BlueprintReadWrite, Category="Editable Mesh" )
+	TArray<FVertexIndexAndInstanceID> VertexIndicesAndInstanceIDs;
+
+	/** Default constructor */
+	FVertexInstancesForPolygonHole()
+		: VertexIndicesAndInstanceIDs()
+	{
+	}
+
+	FString ToString() const
+	{
+		return FString::Printf(
+			TEXT( "VertexInstanceIDs:%s" ),
+			*LogHelpers::ArrayToString( VertexIndicesAndInstanceIDs ) );
+	}
+};
+
+
+USTRUCT( BlueprintType )
+struct FChangeVertexInstancesForPolygon
+{
+	GENERATED_BODY()
+
+	/** The polygon to set vertex instances on */
+	UPROPERTY( BlueprintReadWrite, Category="Editable Mesh" )
+	FPolygonID PolygonID;
+
+	/** A list of vertex numbers and vertex instance IDs to change for the polygon perimeter. */
+	UPROPERTY( BlueprintReadWrite, Category="Editable Mesh" )
+	TArray<FVertexIndexAndInstanceID> PerimeterVertexIndicesAndInstanceIDs;
+
+	/** A list of vertex instance IDs for each hole. */
+	UPROPERTY( BlueprintReadWrite, Category="Editable Mesh" )
+	TArray<FVertexInstancesForPolygonHole> VertexIndicesAndInstanceIDsForEachHole;
+
+	/** Default constructor */
+	FChangeVertexInstancesForPolygon()
+		: PolygonID( FPolygonID::Invalid ),
+		  PerimeterVertexIndicesAndInstanceIDs(),
+		  VertexIndicesAndInstanceIDsForEachHole()
+	{
+	}
+
+	FString ToString() const
+	{
+		return FString::Printf(
+			TEXT( "PolygonID:%s, PerimeterVertexIndicesAndInstanceIDs:%s, VertexIndicesAndInstanceIDsForEachHole:%s" ),
+			*PolygonID.ToString(),
+			*LogHelpers::ArrayToString( PerimeterVertexIndicesAndInstanceIDs ),
+			*LogHelpers::ArrayToString( VertexIndicesAndInstanceIDsForEachHole ) );
+	}
+};
+
+
 USTRUCT( BlueprintType )
 struct FVerticesForEdge
 {
@@ -1142,47 +1286,45 @@ struct FVertexToMove
 
 
 USTRUCT( BlueprintType )
-struct FSectionToCreate
+struct FPolygonGroupToCreate
 {
 	GENERATED_BODY()
 
-	/** Material to assign to the new section */
+	// @todo mesheditor: Material, bEnableCollision and bCastShadow should all be implemented as PolygonGroup attributes.
+	// This would require making Material an asset name, rather than a UObject*, which is also safer.
+
+	/** Material to assign to the new polygon group */
 	UPROPERTY( BlueprintReadWrite, Category="Editable Mesh" )
 	class UMaterialInterface* Material;
 
-	/** Whether the new section should have collision enabled */
+	/** Whether the new polygon group should have collision enabled */
 	UPROPERTY( BlueprintReadWrite, Category="Editable Mesh" )
 	bool bEnableCollision;
 
-	/** Whether the new section casts a shadow */
+	/** Whether the new polygon group casts a shadow */
 	UPROPERTY( BlueprintReadWrite, Category="Editable Mesh" )
 	bool bCastShadow;
 
-	/** The original ID of the section.  Should only be used by the undo system. */
-	UPROPERTY( BlueprintReadWrite, Category="Editable Mesh" )
-	FSectionID OriginalSectionID;
-
-	/** The original rendering section index.  Should only be used by the undo system. */
+	/** The original ID of the polygon group.  Should only be used by the undo system. */
 	UPROPERTY( BlueprintReadWrite, Category="Editable Mesh" )
-	int32 OriginalRenderingSectionIndex;
+	FPolygonGroupID OriginalPolygonGroupID;
 
 	/** Default constructor */
-	FSectionToCreate()
+	FPolygonGroupToCreate()
 		: Material( nullptr )
 		, bEnableCollision( false )
 		, bCastShadow( false )
-		, OriginalSectionID( FSectionID::Invalid )
-		, OriginalRenderingSectionIndex( INDEX_NONE )
+		, OriginalPolygonGroupID( FPolygonGroupID::Invalid )
 	{
 	}
 
 	FString ToString() const
 	{
 		return FString::Printf(
-			TEXT( "Material:%s, bEnableCollision:%s, bCastShadow:%s, OriginalSectionID:%s" ),
+			TEXT( "Material:%s, bEnableCollision:%s, bCastShadow:%s, OriginalPolygonGroupID:%s" ),
 			Material ? *((UObject*)Material )->GetName() : TEXT( "<none>" ),
 			*LexicalConversion::ToString( bEnableCollision ),
 			*LexicalConversion::ToString( bCastShadow ),
-			*OriginalSectionID.ToString() );
+			*OriginalPolygonGroupID.ToString() );
 	}
 };
diff --git a/Engine/Source/Runtime/MeshEditingRuntime/StaticMeshEditableMeshFormat.cpp b/Engine/Source/Runtime/MeshEditingRuntime/StaticMeshEditableMeshFormat.cpp
--- a/Engine/Source/Runtime/MeshEditingRuntime/StaticMeshEditableMeshFormat.cpp
+++ b/Engine/Source/Runtime/MeshEditingRuntime/StaticMeshEditableMeshFormat.cpp
@@ -5,7 +5,7 @@
 #include "Engine/StaticMesh.h"
 #include "Components/StaticMeshComponent.h"
 #include "StaticMeshResources.h"
-#include "EditableStaticMesh.h"
+#include "EditableStaticMeshAdapter.h"
 
 
 void FStaticMeshEditableMeshFormat::FillMeshObjectPtr( UPrimitiveComponent& Component, FEditableMeshSubMeshAddress& SubMeshAddress )
@@ -28,7 +28,7 @@ void FStaticMeshEditableMeshFormat::FillMeshObjectPtr( UPrimitiveComponent& Comp
 
 UEditableMesh* FStaticMeshEditableMeshFormat::MakeEditableMesh( UPrimitiveComponent& Component, const FEditableMeshSubMeshAddress& SubMeshAddress )
 {
-	// If the static mesh already has an attached UEditableStaticMesh, use that in preference to creating a new one
+	// If the static mesh already has an attached UEditableStaticMeshAdapter, use that in preference to creating a new one
 	const UStaticMeshComponent* StaticMeshComponentPtr = Cast<const UStaticMeshComponent>( &Component );
 	if( StaticMeshComponentPtr != nullptr )
 	{
@@ -44,17 +44,21 @@ UEditableMesh* FStaticMeshEditableMeshFormat::MakeEditableMesh( UPrimitiveCompon
 		}
 	}
 
-	UEditableStaticMesh* EditableStaticMesh = NewObject<UEditableStaticMesh>();
+	UEditableMesh* EditableMesh = NewObject<UEditableMesh>();
+	UEditableStaticMeshAdapter* EditableStaticMesh = NewObject<UEditableStaticMeshAdapter>( EditableMesh );
+	EditableMesh->Adapters.Add( EditableStaticMesh );
 
-	EditableStaticMesh->InitEditableStaticMesh( Component, SubMeshAddress );
+	EditableStaticMesh->InitEditableStaticMesh( EditableMesh, Component, SubMeshAddress );
 
 	// Don't bother returning a new mesh if it has no geometry
-	if( EditableStaticMesh->GetVertexCount() == 0 )
+	if( EditableMesh->GetVertexCount() == 0 )
 	{
+		EditableMesh->Adapters.Remove( EditableStaticMesh );
 		EditableStaticMesh->MarkPendingKill();
-		EditableStaticMesh = nullptr;
+		EditableMesh->MarkPendingKill();
+		EditableMesh = nullptr;
 	}
 
-	return EditableStaticMesh;
+	return EditableMesh;
 }
 
