Fixed issue where vertex normal/tangent setting was no longer going through the 'proper' SetVertexInstancesAttributes route, and hence were not having their changes seen by the adapter.

diff --git a/Engine/Source/Runtime/MeshEditingRuntime/EditableMesh.cpp b/Engine/Source/Runtime/MeshEditingRuntime/EditableMesh.cpp
--- a/Engine/Source/Runtime/MeshEditingRuntime/EditableMesh.cpp
+++ b/Engine/Source/Runtime/MeshEditingRuntime/EditableMesh.cpp
@@ -7022,6 +7022,9 @@ void UEditableMesh::GenerateTangentsAndNormals()
 		}
 	}
 
+	static TArray<FAttributesForVertexInstance> AttributesForVertexInstances;
+	AttributesForVertexInstances.Reset( VertexInstanceIDs.Num() );
+
 	for( const FVertexInstanceID VertexInstanceID : VertexInstanceIDs )
 	{
 		FVector Normal = FVector::ZeroVector;
@@ -7069,10 +7072,17 @@ void UEditableMesh::GenerateTangentsAndNormals()
 		//UE_LOG( LogEditableMesh, Verbose, TEXT( "%s tangent: old %s" ), *VertexInstanceID.ToString(), *VertexInstance.Tangent.ToString() );
 		//UE_LOG( LogEditableMesh, Verbose, TEXT( "%s tangent: new %s" ), *VertexInstanceID.ToString(), *Tangent.ToString() );
 
-		VertexInstance.Normal = Normal;
-		VertexInstance.Tangent = Tangent;
-		VertexInstance.BinormalSign = BinormalSign;
+		AttributesForVertexInstances.Emplace();
+		FAttributesForVertexInstance& AttributesForVertexInstance = AttributesForVertexInstances.Last();
+
+		AttributesForVertexInstance.VertexInstanceID = VertexInstanceID;
+		AttributesForVertexInstance.VertexInstanceAttributes.Attributes.Reset( 3 );
+		AttributesForVertexInstance.VertexInstanceAttributes.Attributes.Emplace( UEditableMeshAttribute::VertexNormal(), 0, FVector4( Normal ) );
+		AttributesForVertexInstance.VertexInstanceAttributes.Attributes.Emplace( UEditableMeshAttribute::VertexTangent(), 0, FVector4( Tangent ) );
+		AttributesForVertexInstance.VertexInstanceAttributes.Attributes.Emplace( UEditableMeshAttribute::VertexBinormalSign(), 0, FVector4( BinormalSign, 0.0f, 0.0f, 0.0f ) );
 	}
+
+	SetVertexInstancesAttributes( AttributesForVertexInstances );
 }
 
 
