Removed dead code.

diff --git a/Engine/Source/Runtime/MeshEditingRuntime/EditableStaticMeshAdapter.cpp b/Engine/Source/Runtime/MeshEditingRuntime/EditableStaticMeshAdapter.cpp
--- a/Engine/Source/Runtime/MeshEditingRuntime/EditableStaticMeshAdapter.cpp
+++ b/Engine/Source/Runtime/MeshEditingRuntime/EditableStaticMeshAdapter.cpp
@@ -1112,24 +1112,6 @@ void UEditableStaticMeshAdapter::OnSetEdgeAttribute( const UEditableMesh* Editab
 }
 
 
-//int32 UEditableStaticMeshAdapter::GetTriangleCount( const FPolygonGroupID PolygonGroupID ) const
-//{
-//	checkSlow( RenderingPolygonGroups.IsAllocated( PolygonGroupID.GetValue() ) );
-//	const FRenderingPolygonGroup& RenderingPolygonGroup = RenderingPolygonGroups[ PolygonGroupID.GetValue() ];
-//
-//	return RenderingPolygonGroup.Triangles.Num();
-//}
-//
-//
-//int32 UEditableStaticMeshAdapter::GetTriangleArraySize( const FPolygonGroupID PolygonGroupID ) const
-//{
-//	checkSlow( RenderingPolygonGroups.IsAllocated( PolygonGroupID.GetValue() ) );
-//	const FRenderingPolygonGroup& RenderingPolygonGroup = RenderingPolygonGroups[ PolygonGroupID.GetValue() ];
-//
-//	return RenderingPolygonGroup.Triangles.GetMaxIndex();
-//}
-
-
 void UEditableStaticMeshAdapter::OnSetVertexInstanceAttribute( const UEditableMesh* EditableMesh, const FVertexInstanceID VertexInstanceID, const FName AttributeName, const int32 AttributeIndex, const FVector4 AttributeValue )
 {
 	const FMeshVertexInstance& VertexInstance = EditableMesh->VertexInstances[ VertexInstanceID.GetValue() ];
