Editable Mesh: FMeshTriangle is now compatible with Blueprints
- Made FMeshTriangle a BlueprintType and eliminate unsupported use of static arrays
- Updated existing use cases to either use accessor functions or directly access specific values

diff --git a/Engine/Source/Runtime/MeshEditingRuntime/EditableMesh.cpp b/Engine/Source/Runtime/MeshEditingRuntime/EditableMesh.cpp
--- a/Engine/Source/Runtime/MeshEditingRuntime/EditableMesh.cpp
+++ b/Engine/Source/Runtime/MeshEditingRuntime/EditableMesh.cpp
@@ -242,9 +242,11 @@ void UEditableMesh::FixUpElementIDs( const FElementIDRemappings& Remappings )
 
 		for( FMeshTriangle& Triangle : Polygon.Triangles )
 		{
-			for( FVertexInstanceID& VertexInstanceID : Triangle.VertexInstanceIDs )
+			for( int32 TriangleVertexNumber = 0; TriangleVertexNumber < 3; ++TriangleVertexNumber )
 			{
-				VertexInstanceID = Remappings.GetRemappedVertexInstanceID( VertexInstanceID );
+				const FVertexInstanceID OriginalVertexInstanceID = Triangle.GetVertexInstanceID( TriangleVertexNumber );
+				const FVertexInstanceID NewVertexInstanceID = Remappings.GetRemappedVertexInstanceID( OriginalVertexInstanceID );
+				Triangle.SetVertexInstanceID( TriangleVertexNumber, NewVertexInstanceID );
 			}
 		}
 
@@ -1796,7 +1798,7 @@ FVector UEditableMesh::GetPolygonTriangulatedTriangleVertexPosition( const FPoly
 	checkSlow( Polygons.IsAllocated( PolygonID.GetValue() ) );
 	const FMeshPolygon& Polygon = Polygons[ PolygonID.GetValue() ];
 
-	const FVertexInstanceID VertexInstanceID = Polygon.Triangles[ PolygonTriangleNumber ].VertexInstanceIDs[ TriangleVertexNumber ];
+	const FVertexInstanceID VertexInstanceID = Polygon.Triangles[ PolygonTriangleNumber ].GetVertexInstanceID( TriangleVertexNumber );
 	checkSlow( VertexInstances.IsAllocated( VertexInstanceID.GetValue() ) );
 	const FVertexID VertexID = VertexInstances[ VertexInstanceID.GetValue() ].VertexID;
 
@@ -2716,9 +2718,9 @@ void UEditableMesh::ComputePolygonTriangulation( const FPolygonID PolygonID, TAr
 				OutTriangles.Emplace();
 				FMeshTriangle& Triangle = OutTriangles.Last();
 
-				Triangle.VertexInstanceIDs[ 0 ] = PolygonVertexInstanceIDs[ PrevVertexNumbers[ EarVertexNumber ] ];
-				Triangle.VertexInstanceIDs[ 1 ] = PolygonVertexInstanceIDs[ EarVertexNumber ];
-				Triangle.VertexInstanceIDs[ 2 ] = PolygonVertexInstanceIDs[ NextVertexNumbers[ EarVertexNumber ] ];
+				Triangle.SetVertexInstanceID( 0, PolygonVertexInstanceIDs[ PrevVertexNumbers[ EarVertexNumber ] ] );
+				Triangle.SetVertexInstanceID( 1, PolygonVertexInstanceIDs[ EarVertexNumber ] );
+				Triangle.SetVertexInstanceID( 2, PolygonVertexInstanceIDs[ NextVertexNumbers[ EarVertexNumber ] ] );
 			}
 
 			// Update our linked list.  We're effectively cutting off the ear by pointing the ear vertex's neighbors to
@@ -2760,9 +2762,9 @@ bool UEditableMesh::ComputeBarycentricWeightForPointOnPolygon( const FPolygonID
 	// Figure out which triangle the incoming point is within
 	for( const FMeshTriangle& Triangle : Polygon.Triangles )
 	{
-		const FVector TriangleVertex0Position = GetVertexInstanceAttribute( Triangle.VertexInstanceIDs[ 0 ], UEditableMeshAttribute::VertexPosition(), 0 );
-		const FVector TriangleVertex1Position = GetVertexInstanceAttribute( Triangle.VertexInstanceIDs[ 1 ], UEditableMeshAttribute::VertexPosition(), 0 );
-		const FVector TriangleVertex2Position = GetVertexInstanceAttribute( Triangle.VertexInstanceIDs[ 2 ], UEditableMeshAttribute::VertexPosition(), 0 );
+		const FVector TriangleVertex0Position = GetVertexInstanceAttribute( Triangle.VertexInstanceID0, UEditableMeshAttribute::VertexPosition(), 0 );
+		const FVector TriangleVertex1Position = GetVertexInstanceAttribute( Triangle.VertexInstanceID1, UEditableMeshAttribute::VertexPosition(), 0 );
+		const FVector TriangleVertex2Position = GetVertexInstanceAttribute( Triangle.VertexInstanceID2, UEditableMeshAttribute::VertexPosition(), 0 );
 
 		// Calculate the barycentric weights for the triangle's verts and determine if the point lies within its bounds.
 		OutTriangleVertexWeights = FMath::ComputeBaryCentric2D( PointOnPolygon, TriangleVertex0Position, TriangleVertex1Position, TriangleVertex2Position );
@@ -2793,9 +2795,9 @@ void UEditableMesh::ComputeTextureCoordinatesForPointOnPolygon( const FPolygonID
 		for( int32 TextureCoordinateIndex = 0; TextureCoordinateIndex < TextureCoordinateCount; ++TextureCoordinateIndex )
 		{
 			OutInterpolatedTextureCoordinates[ TextureCoordinateIndex ] =
-				TriangleVertexWeights.X * GetVertexInstanceAttribute( Triangle.VertexInstanceIDs[ 0 ], UEditableMeshAttribute::VertexTextureCoordinate(), TextureCoordinateIndex ) +
-				TriangleVertexWeights.Y * GetVertexInstanceAttribute( Triangle.VertexInstanceIDs[ 1 ], UEditableMeshAttribute::VertexTextureCoordinate(), TextureCoordinateIndex ) +
-				TriangleVertexWeights.Z * GetVertexInstanceAttribute( Triangle.VertexInstanceIDs[ 2 ], UEditableMeshAttribute::VertexTextureCoordinate(), TextureCoordinateIndex );
+				TriangleVertexWeights.X * GetVertexInstanceAttribute( Triangle.VertexInstanceID0, UEditableMeshAttribute::VertexTextureCoordinate(), TextureCoordinateIndex ) +
+				TriangleVertexWeights.Y * GetVertexInstanceAttribute( Triangle.VertexInstanceID1, UEditableMeshAttribute::VertexTextureCoordinate(), TextureCoordinateIndex ) +
+				TriangleVertexWeights.Z * GetVertexInstanceAttribute( Triangle.VertexInstanceID2, UEditableMeshAttribute::VertexTextureCoordinate(), TextureCoordinateIndex );
 		}
 
 		bOutWereTextureCoordinatesFound = true;
@@ -5059,9 +5061,9 @@ void UEditableMesh::ChangePolygonsVertexInstances( const TArray<FChangeVertexIns
 				{
 					for( int32 VertexIndex = 0; VertexIndex < 3; ++VertexIndex )
 					{
-						if( Triangle.VertexInstanceIDs[ VertexIndex ] == OldVertexInstanceID )
+						if( Triangle.GetVertexInstanceID( VertexIndex ) == OldVertexInstanceID )
 						{
-							Triangle.VertexInstanceIDs[ VertexIndex ] = IndexAndInstance.VertexInstanceID;
+							Triangle.SetVertexInstanceID( VertexIndex, IndexAndInstance.VertexInstanceID );
 						}
 					}
 				}
@@ -5088,9 +5090,9 @@ void UEditableMesh::ChangePolygonsVertexInstances( const TArray<FChangeVertexIns
 					{
 						for( int32 VertexIndex = 0; VertexIndex < 3; ++VertexIndex )
 						{
-							if( Triangle.VertexInstanceIDs[ VertexIndex ] == OldVertexInstanceID )
+							if( Triangle.GetVertexInstanceID( VertexIndex ) == OldVertexInstanceID )
 							{
-								Triangle.VertexInstanceIDs[ VertexIndex ] = IndexAndInstance.VertexInstanceID;
+								Triangle.SetVertexInstanceID( VertexIndex, IndexAndInstance.VertexInstanceID );
 							}
 						}
 					}
@@ -6677,15 +6679,15 @@ void UEditableMesh::BevelOrInsetPolygons( const TArray<FPolygonID>& PolygonIDs,
 				for( int32 TextureCoordinateIndex = 0; TextureCoordinateIndex < TextureCoordinateCount; ++TextureCoordinateIndex )
 				{
 					TextureCoordinatesForNewVertex[ TextureCoordinateIndex ] =
-						TriangleVertexWeights.X * GetVertexInstanceAttribute( Triangle.VertexInstanceIDs[ 0 ], UEditableMeshAttribute::VertexTextureCoordinate(), TextureCoordinateIndex ) +
-						TriangleVertexWeights.Y * GetVertexInstanceAttribute( Triangle.VertexInstanceIDs[ 1 ], UEditableMeshAttribute::VertexTextureCoordinate(), TextureCoordinateIndex ) +
-						TriangleVertexWeights.Z * GetVertexInstanceAttribute( Triangle.VertexInstanceIDs[ 2 ], UEditableMeshAttribute::VertexTextureCoordinate(), TextureCoordinateIndex );
+						TriangleVertexWeights.X * GetVertexInstanceAttribute( Triangle.VertexInstanceID0, UEditableMeshAttribute::VertexTextureCoordinate(), TextureCoordinateIndex ) +
+						TriangleVertexWeights.Y * GetVertexInstanceAttribute( Triangle.VertexInstanceID1, UEditableMeshAttribute::VertexTextureCoordinate(), TextureCoordinateIndex ) +
+						TriangleVertexWeights.Z * GetVertexInstanceAttribute( Triangle.VertexInstanceID2, UEditableMeshAttribute::VertexTextureCoordinate(), TextureCoordinateIndex );
 				}
 
 				VertexColorsForNewVertices.Add(
-					TriangleVertexWeights.X * GetVertexInstanceAttribute( Triangle.VertexInstanceIDs[ 0 ], UEditableMeshAttribute::VertexColor(), 0 ) +
-					TriangleVertexWeights.Y * GetVertexInstanceAttribute( Triangle.VertexInstanceIDs[ 1 ], UEditableMeshAttribute::VertexColor(), 0 ) +
-					TriangleVertexWeights.Z * GetVertexInstanceAttribute( Triangle.VertexInstanceIDs[ 2 ], UEditableMeshAttribute::VertexColor(), 0 )
+					TriangleVertexWeights.X * GetVertexInstanceAttribute( Triangle.VertexInstanceID0, UEditableMeshAttribute::VertexColor(), 0 ) +
+					TriangleVertexWeights.Y * GetVertexInstanceAttribute( Triangle.VertexInstanceID1, UEditableMeshAttribute::VertexColor(), 0 ) +
+					TriangleVertexWeights.Z * GetVertexInstanceAttribute( Triangle.VertexInstanceID2, UEditableMeshAttribute::VertexColor(), 0 )
 				);
 			}
 			else
@@ -6970,9 +6972,9 @@ void UEditableMesh::GeneratePolygonTangentsAndNormals( const TArray<FPolygonID>&
 
 		for( const FMeshTriangle& Triangle : Polygon.Triangles )
 		{
-			const FMeshVertexInstance& VertexInstance0 = VertexInstances[ Triangle.VertexInstanceIDs[ 0 ].GetValue() ];
-			const FMeshVertexInstance& VertexInstance1 = VertexInstances[ Triangle.VertexInstanceIDs[ 1 ].GetValue() ];
-			const FMeshVertexInstance& VertexInstance2 = VertexInstances[ Triangle.VertexInstanceIDs[ 2 ].GetValue() ];
+			const FMeshVertexInstance& VertexInstance0 = VertexInstances[ Triangle.VertexInstanceID0.GetValue() ];
+			const FMeshVertexInstance& VertexInstance1 = VertexInstances[ Triangle.VertexInstanceID1.GetValue() ];
+			const FMeshVertexInstance& VertexInstance2 = VertexInstances[ Triangle.VertexInstanceID2.GetValue() ];
 
 			const FMeshVertex& Vertex0 = Vertices[ VertexInstance0.VertexID.GetValue() ];
 			const FMeshVertex& Vertex1 = Vertices[ VertexInstance1.VertexID.GetValue() ];
@@ -7676,7 +7678,7 @@ void UEditableMesh::TriangulatePolygons( const TArray<FPolygonID>& PolygonIDs, T
 						PolygonToCreate.PerimeterVertices.Emplace();
 						FVertexAndAttributes& PerimeterVertex = PolygonToCreate.PerimeterVertices.Last();
 
-						PerimeterVertex.VertexInstanceID = Triangle.VertexInstanceIDs[ TriangleVertexNumber ];
+						PerimeterVertex.VertexInstanceID = Triangle.GetVertexInstanceID( TriangleVertexNumber );
 						PerimeterVertex.VertexID = FVertexID::Invalid;
 					}
 				}
@@ -8247,17 +8249,17 @@ void UEditableMesh::TessellatePolygons( const TArray<FPolygonID>& PolygonIDs, co
 									VertexAndAttributes.PolygonVertexAttributes.Attributes.Add( FMeshElementAttributeData(
 										UEditableMeshAttribute::VertexTextureCoordinate(),
 										TextureCoordinateIndex,
-										TriangleVertexWeights.X * GetVertexInstanceAttribute( Triangle.VertexInstanceIDs[ 0 ], UEditableMeshAttribute::VertexTextureCoordinate(), TextureCoordinateIndex ) +
-										TriangleVertexWeights.Y * GetVertexInstanceAttribute( Triangle.VertexInstanceIDs[ 1 ], UEditableMeshAttribute::VertexTextureCoordinate(), TextureCoordinateIndex ) +
-										TriangleVertexWeights.Z * GetVertexInstanceAttribute( Triangle.VertexInstanceIDs[ 2 ], UEditableMeshAttribute::VertexTextureCoordinate(), TextureCoordinateIndex ) ) );
+										TriangleVertexWeights.X * GetVertexInstanceAttribute( Triangle.VertexInstanceID0, UEditableMeshAttribute::VertexTextureCoordinate(), TextureCoordinateIndex ) +
+										TriangleVertexWeights.Y * GetVertexInstanceAttribute( Triangle.VertexInstanceID1, UEditableMeshAttribute::VertexTextureCoordinate(), TextureCoordinateIndex ) +
+										TriangleVertexWeights.Z * GetVertexInstanceAttribute( Triangle.VertexInstanceID2, UEditableMeshAttribute::VertexTextureCoordinate(), TextureCoordinateIndex ) ) );
 								}
 
 								VertexAndAttributes.PolygonVertexAttributes.Attributes.Add( FMeshElementAttributeData(
 									UEditableMeshAttribute::VertexColor(),
 									0,
-									TriangleVertexWeights.X * GetVertexInstanceAttribute( Triangle.VertexInstanceIDs[ 0 ], UEditableMeshAttribute::VertexColor(), 0 ) +
-									TriangleVertexWeights.Y * GetVertexInstanceAttribute( Triangle.VertexInstanceIDs[ 1 ], UEditableMeshAttribute::VertexColor(), 0 ) +
-									TriangleVertexWeights.Z * GetVertexInstanceAttribute( Triangle.VertexInstanceIDs[ 2 ], UEditableMeshAttribute::VertexColor(), 0 ) ) );
+									TriangleVertexWeights.X * GetVertexInstanceAttribute( Triangle.VertexInstanceID0, UEditableMeshAttribute::VertexColor(), 0 ) +
+									TriangleVertexWeights.Y * GetVertexInstanceAttribute( Triangle.VertexInstanceID1, UEditableMeshAttribute::VertexColor(), 0 ) +
+									TriangleVertexWeights.Z * GetVertexInstanceAttribute( Triangle.VertexInstanceID2, UEditableMeshAttribute::VertexColor(), 0 ) ) );
 							}
 						}
 					}
@@ -8336,17 +8338,17 @@ void UEditableMesh::TessellatePolygons( const TArray<FPolygonID>& PolygonIDs, co
 										VertexAndAttributes.PolygonVertexAttributes.Attributes.Add( FMeshElementAttributeData(
 											UEditableMeshAttribute::VertexTextureCoordinate(),
 											TextureCoordinateIndex,
-											TriangleVertexWeights.X * GetVertexInstanceAttribute( Triangle.VertexInstanceIDs[ 0 ], UEditableMeshAttribute::VertexTextureCoordinate(), TextureCoordinateIndex ) +
-											TriangleVertexWeights.Y * GetVertexInstanceAttribute( Triangle.VertexInstanceIDs[ 1 ], UEditableMeshAttribute::VertexTextureCoordinate(), TextureCoordinateIndex ) +
-											TriangleVertexWeights.Z * GetVertexInstanceAttribute( Triangle.VertexInstanceIDs[ 2 ], UEditableMeshAttribute::VertexTextureCoordinate(), TextureCoordinateIndex ) ) );
+											TriangleVertexWeights.X * GetVertexInstanceAttribute( Triangle.VertexInstanceID0, UEditableMeshAttribute::VertexTextureCoordinate(), TextureCoordinateIndex ) +
+											TriangleVertexWeights.Y * GetVertexInstanceAttribute( Triangle.VertexInstanceID1, UEditableMeshAttribute::VertexTextureCoordinate(), TextureCoordinateIndex ) +
+											TriangleVertexWeights.Z * GetVertexInstanceAttribute( Triangle.VertexInstanceID2, UEditableMeshAttribute::VertexTextureCoordinate(), TextureCoordinateIndex ) ) );
 									}
 
 									VertexAndAttributes.PolygonVertexAttributes.Attributes.Add( FMeshElementAttributeData(
 										UEditableMeshAttribute::VertexColor(),
 										0,
-										TriangleVertexWeights.X * GetVertexInstanceAttribute( Triangle.VertexInstanceIDs[ 0 ], UEditableMeshAttribute::VertexColor(), 0 ) +
-										TriangleVertexWeights.Y * GetVertexInstanceAttribute( Triangle.VertexInstanceIDs[ 1 ], UEditableMeshAttribute::VertexColor(), 0 ) +
-										TriangleVertexWeights.Z * GetVertexInstanceAttribute( Triangle.VertexInstanceIDs[ 2 ], UEditableMeshAttribute::VertexColor(), 0 ) ) );
+										TriangleVertexWeights.X * GetVertexInstanceAttribute( Triangle.VertexInstanceID0, UEditableMeshAttribute::VertexColor(), 0 ) +
+										TriangleVertexWeights.Y * GetVertexInstanceAttribute( Triangle.VertexInstanceID1, UEditableMeshAttribute::VertexColor(), 0 ) +
+										TriangleVertexWeights.Z * GetVertexInstanceAttribute( Triangle.VertexInstanceID2, UEditableMeshAttribute::VertexColor(), 0 ) ) );
 								}
 							}
 						}
diff --git a/Engine/Source/Runtime/MeshEditingRuntime/EditableStaticMeshAdapter.cpp b/Engine/Source/Runtime/MeshEditingRuntime/EditableStaticMeshAdapter.cpp
--- a/Engine/Source/Runtime/MeshEditingRuntime/EditableStaticMeshAdapter.cpp
+++ b/Engine/Source/Runtime/MeshEditingRuntime/EditableStaticMeshAdapter.cpp
@@ -42,8 +42,9 @@ inline void UEditableStaticMeshAdapter::UpdateIndexBufferFormatIfNeeded( const T
 	{
 		for( const FMeshTriangle& Triangle : Triangles )
 		{
-			for( const FVertexInstanceID VertexInstanceID : Triangle.VertexInstanceIDs )
+			for( int32 TriangleVertexNumber = 0; TriangleVertexNumber < 3; ++TriangleVertexNumber )
 			{
+				const FVertexInstanceID VertexInstanceID = Triangle.GetVertexInstanceID( TriangleVertexNumber );
 				if( VertexInstanceID.GetValue() > TNumericLimits<uint16>::Max() )
 				{
 					EnsureIndexBufferIs32Bit();
@@ -342,7 +343,7 @@ void UEditableStaticMeshAdapter::InitEditableStaticMesh( UEditableMesh* Editable
 								VertexInstance.ConnectedPolygons.Add( NewPolygonID );
 
 								// The triangle points to each of its three vertices
-								NewTriangle.VertexInstanceIDs[ TriangleVertexIndex ] = VertexInstanceID;
+								NewTriangle.SetVertexInstanceID( TriangleVertexIndex, VertexInstanceID );
 							}
 
 							// Add triangle to polygon triangulation array
@@ -679,7 +680,7 @@ void UEditableStaticMeshAdapter::OnRebuildRenderMesh( const UEditableMesh* Edita
 
 				// Find the first valid vertex instance index, so that we have a value we can use for our degenerates
 				check( RenderingPolygonGroup.Triangles.Num() > 0 );
-				const FVertexInstanceID FirstValidRenderingID = RenderingPolygonGroup.Triangles.CreateConstIterator()->VertexInstanceIDs[ 0 ];
+				const FVertexInstanceID FirstValidRenderingID = RenderingPolygonGroup.Triangles.CreateConstIterator()->GetVertexInstanceID( 0 );
 
 				for( int32 TriangleIndex = 0; TriangleIndex < RenderingPolygonGroup.Triangles.GetMaxIndex(); ++TriangleIndex )
 				{
@@ -688,7 +689,7 @@ void UEditableStaticMeshAdapter::OnRebuildRenderMesh( const UEditableMesh* Edita
 						const FMeshTriangle& Triangle = RenderingPolygonGroup.Triangles[ TriangleIndex ];
 						for( int32 TriVert = 0; TriVert < 3; ++TriVert )
 						{
-							const uint32 RenderingVertexIndex = Triangle.VertexInstanceIDs[ TriVert ].GetValue();
+							const uint32 RenderingVertexIndex = Triangle.GetVertexInstanceID( TriVert ).GetValue();
 							IndexBuffer.Add( RenderingVertexIndex );
 							MinIndex = FMath::Min( MinIndex, RenderingVertexIndex );
 							MaxIndex = FMath::Max( MaxIndex, RenderingVertexIndex );
@@ -844,9 +845,11 @@ void UEditableStaticMeshAdapter::OnReindexElements( const UEditableMesh* Editabl
 
 		for( FMeshTriangle& Triangle : RenderingPolygonGroup.Triangles )
 		{
-			for( FVertexInstanceID& VertexInstanceID : Triangle.VertexInstanceIDs )
+			for( int32 TriangleVertexNumber = 0; TriangleVertexNumber < 3; ++TriangleVertexNumber )
 			{
-				VertexInstanceID = Remappings.NewVertexInstanceIndexLookup[ VertexInstanceID.GetValue() ];
+				const FVertexInstanceID OriginalVertexInstanceID = Triangle.GetVertexInstanceID( TriangleVertexNumber );
+				const FVertexInstanceID NewVertexInstanceID = Remappings.NewVertexInstanceIndexLookup[ OriginalVertexInstanceID.GetValue() ];
+				Triangle.SetVertexInstanceID( TriangleVertexNumber, NewVertexInstanceID );
 			}
 		}
 
@@ -1305,9 +1308,9 @@ void UEditableStaticMeshAdapter::OnRetriangulatePolygons( const UEditableMesh* E
 					const FMeshTriangle& OldTriangle = RenderingPolygonGroup.Triangles[ RenderingPolygon.TriangulatedPolygonTriangleIndices[ Index ].GetValue() ];
 					const FMeshTriangle& NewTriangle = Triangles[ Index ];
 
-					if( OldTriangle.VertexInstanceIDs[ 0 ] != NewTriangle.VertexInstanceIDs[ 0 ] ||
-						OldTriangle.VertexInstanceIDs[ 1 ] != NewTriangle.VertexInstanceIDs[ 1 ] ||
-						OldTriangle.VertexInstanceIDs[ 2 ] != NewTriangle.VertexInstanceIDs[ 2 ] )
+					if( OldTriangle.VertexInstanceID0 != NewTriangle.VertexInstanceID0 ||
+						OldTriangle.VertexInstanceID1 != NewTriangle.VertexInstanceID1 ||
+						OldTriangle.VertexInstanceID2 != NewTriangle.VertexInstanceID2 )
 					{
 						bNeedsUpdatedTriangles = true;
 						break;
@@ -1363,8 +1366,8 @@ void UEditableStaticMeshAdapter::OnRetriangulatePolygons( const UEditableMesh* E
 						FMeshTriangle& NewTriangle = RenderingPolygonGroup.Triangles[ NewTriangleIndex ];
 						for( int32 TriangleVertexNumber = 0; TriangleVertexNumber < 3; ++TriangleVertexNumber )
 						{
-							const FVertexInstanceID VertexInstanceID = Triangles[ TriangleToAddNumber ].VertexInstanceIDs[ TriangleVertexNumber ];
-							NewTriangle.VertexInstanceIDs[ TriangleVertexNumber ] = VertexInstanceID;
+							const FVertexInstanceID VertexInstanceID = Triangles[ TriangleToAddNumber ].GetVertexInstanceID( TriangleVertexNumber );
+							NewTriangle.SetVertexInstanceID( TriangleVertexNumber, VertexInstanceID );
 							MinVertexIndex = FMath::Min( MinVertexIndex, VertexInstanceID.GetValue() );
 							MaxVertexIndex = FMath::Max( MaxVertexIndex, VertexInstanceID.GetValue() );
 						}
@@ -1417,7 +1420,7 @@ void UEditableStaticMeshAdapter::OnRetriangulatePolygons( const UEditableMesh* E
 						{
 							StaticMeshLOD.IndexBuffer.SetIndex(
 								FRenderingPolygonGroup::TriangleIndexToRenderingTriangleFirstIndex( RenderingSection, FTriangleID( NewTriangleIndex ) ) + TriangleVertexNumber,
-								RenderingPolygonGroup.Triangles[ NewTriangleIndex ].VertexInstanceIDs[ TriangleVertexNumber ].GetValue() );
+								RenderingPolygonGroup.Triangles[ NewTriangleIndex ].GetVertexInstanceID( TriangleVertexNumber ).GetValue() );
 						}
 					}
 
@@ -1495,8 +1498,9 @@ void UEditableStaticMeshAdapter::DeletePolygonTriangles( const UEditableMesh* Ed
 			{
 				const FMeshTriangle& Triangle = RenderingPolygonGroup.Triangles[ TriangleIndexToRemove.GetValue() ];
 
-				for( const FVertexInstanceID VertexInstanceID : Triangle.VertexInstanceIDs )
+				for( int32 TriangleVertexNumber = 0; TriangleVertexNumber < 3; ++TriangleVertexNumber )
 				{
+					const FVertexInstanceID VertexInstanceID = Triangle.GetVertexInstanceID( TriangleVertexNumber );
 					if( VertexInstanceID.GetValue() == MinVertexIndex || VertexInstanceID.GetValue() == MaxVertexIndex )
 					{
 						bUpdateMinMax = true;
@@ -1529,8 +1533,9 @@ void UEditableStaticMeshAdapter::DeletePolygonTriangles( const UEditableMesh* Ed
 
 				for( const FMeshTriangle& Triangle : RenderingPolygonGroup.Triangles )
 				{
-					for( const FVertexInstanceID VertexInstanceID : Triangle.VertexInstanceIDs )
+					for( int32 TriangleVertexNumber = 0; TriangleVertexNumber < 3; ++TriangleVertexNumber )
 					{
+						const FVertexInstanceID VertexInstanceID = Triangle.GetVertexInstanceID( TriangleVertexNumber );
 						if( VertexInstanceID.GetValue() < MinVertexIndex )
 						{
 							MinVertexIndex = VertexInstanceID.GetValue();
diff --git a/Engine/Source/Runtime/MeshEditingRuntime/Public/EditableMesh.h b/Engine/Source/Runtime/MeshEditingRuntime/Public/EditableMesh.h
--- a/Engine/Source/Runtime/MeshEditingRuntime/Public/EditableMesh.h
+++ b/Engine/Source/Runtime/MeshEditingRuntime/Public/EditableMesh.h
@@ -271,21 +271,43 @@ struct FMeshPolygonContour
 };
 
 
-USTRUCT()
+USTRUCT( BlueprintType )
 struct FMeshTriangle
 {
 	GENERATED_BODY()
 
-	/** The three vertex instances that make up this triangle, ordered counter-clockwise. */
-	UPROPERTY()
-	FVertexInstanceID VertexInstanceIDs[ 3 ];
+	/** First vertex instance that makes up this triangle.  Indices must be ordered counter-clockwise. */
+	UPROPERTY( BlueprintReadWrite, Category="Editable Mesh" )
+	FVertexInstanceID VertexInstanceID0;
+
+	/** Second vertex instance that makes up this triangle.  Indices must be ordered counter-clockwise. */
+	UPROPERTY( BlueprintReadWrite, Category="Editable Mesh" )
+	FVertexInstanceID VertexInstanceID1;
+
+	/** Third vertex instance that makes up this triangle.  Indices must be ordered counter-clockwise. */
+	UPROPERTY( BlueprintReadWrite, Category="Editable Mesh" )
+	FVertexInstanceID VertexInstanceID2;
+
+	/** Gets the specified triangle vertex instance ID.  Pass an index between 0 and 2 inclusive. */
+	inline FVertexInstanceID GetVertexInstanceID( const int32 Index ) const
+	{
+		checkSlow( Index >= 0 && Index <= 2 );
+		return reinterpret_cast<const FVertexInstanceID*>( this )[ Index ];
+	}
+
+	/** Sets the specified triangle vertex instance ID.  Pass an index between 0 and 2 inclusive, and the new vertex instance ID to store. */
+	inline void SetVertexInstanceID( const int32 Index, const FVertexInstanceID NewVertexInstanceID )
+	{
+		checkSlow( Index >= 0 && Index <= 2 );
+		( reinterpret_cast<FVertexInstanceID*>( this )[ Index ] ) = NewVertexInstanceID;
+	}
 
 	/** Serializer */
 	friend FArchive& operator<<( FArchive& Ar, FMeshTriangle& Tri )
 	{
-		Ar << Tri.VertexInstanceIDs[ 0 ];
-		Ar << Tri.VertexInstanceIDs[ 1 ];
-		Ar << Tri.VertexInstanceIDs[ 2 ];
+		Ar << Tri.VertexInstanceID0;
+		Ar << Tri.VertexInstanceID1;
+		Ar << Tri.VertexInstanceID2;
 		return Ar;
 	}
 };
